UE4 - prevent cascading errors in automated testing

diff --git a/Engine/Source/Runtime/AutomationWorker/Private/AutomationWorkerModule.cpp b/Engine/Source/Runtime/AutomationWorker/Private/AutomationWorkerModule.cpp
--- a/Engine/Source/Runtime/AutomationWorker/Private/AutomationWorkerModule.cpp
+++ b/Engine/Source/Runtime/AutomationWorker/Private/AutomationWorkerModule.cpp
@@ -492,7 +492,7 @@ static struct FQueueTests
 		}
 		else
 		{
-			Ar->Logf(ELogVerbosity::Error, TEXT("...Automation Test Failed (%s)"), *Test);
+			Ar->Logf(ELogVerbosity::Log, TEXT("...Automation Test Failed (%s)"), *Test);
 		}
 		bTestInPogress = false;
 		NumTestsRun++;
