Tests ran using -automationtests= will now run after the asset registry has finished loading.  This does mean that on some larger projects the tests may take longer before starting.  This is better than not testing them at all which was the case previously.

UnrealEngine.cpp - removed the automationtests commandline arguments from being added to the deferredcommands array.
AutomationWorkerModules - When initailized the commandlines are added to the new deferredautomationcommands array.  Wrote a function to exectue the commands.  Bound the function to the asset registry delegate that gets called when the the registry is done loading.

diff --git a/Engine/Source/Runtime/AutomationWorker/Private/AutomationWorkerModule.cpp b/Engine/Source/Runtime/AutomationWorker/Private/AutomationWorkerModule.cpp
--- a/Engine/Source/Runtime/AutomationWorker/Private/AutomationWorkerModule.cpp
+++ b/Engine/Source/Runtime/AutomationWorker/Private/AutomationWorkerModule.cpp
@@ -1,5 +1,6 @@
 // Copyright 1998-2015 Epic Games, Inc. All Rights Reserved.
 
+#include "AssetRegistryModule.h"
 #include "AutomationWorkerPrivatePCH.h"
 
 
@@ -61,6 +62,9 @@ void FAutomationWorkerModule::Tick()
 	{
 		MessageEndpoint->ProcessInbox();
 	}
+
+	// Run any of the automation commands that was obtained during initialization.
+	//RunDeferredAutomationCommands();
 }
 
 
@@ -131,8 +135,22 @@ void FAutomationWorkerModule::Initialize()
 	}
 	ExecutionCount = INDEX_NONE;
 	bExecutingNetworkCommandResults = false;
-}
 
+#if !(UE_BUILD_SHIPPING)
+	//Obtain any command line tests commands that use '-automationtests='.
+	FString AutomationCmds;
+	if (FParse::Value(FCommandLine::Get(), TEXT("AutomationTests="), AutomationCmds, false))
+	{
+		new(DeferredAutomationCommands) FString(FString(TEXT("Automation CommandLineTests ")) + AutomationCmds);
+	}
+
+	FAssetRegistryModule& AssetRegistryModule = FModuleManager::LoadModuleChecked<FAssetRegistryModule>("AssetRegistry");
+	if (AssetRegistryModule.Get().IsLoadingAssets())
+	{
+		AssetRegistryModule.Get().OnFilesLoaded().AddRaw(this, &FAutomationWorkerModule::RunDeferredAutomationCommands);
+	}
+#endif // !(UE_BUILD_SHIPPING)
+}
 
 void FAutomationWorkerModule::ReportNetworkCommandComplete()
 {
@@ -393,6 +411,27 @@ void FAutomationWorkerModule::RunTest(const FString& InTestToRun, const int32 In
 	FAutomationTestFramework::GetInstance().StartTestByName(InTestToRun, InRoleIndex);
 }
 
+void FAutomationWorkerModule::RunDeferredAutomationCommands()
+{
+	// Get the loaded asset registry module.
+	FAssetRegistryModule& AssetRegistryModule = FModuleManager::LoadModuleChecked<FAssetRegistryModule>(TEXT("AssetRegistry"));
+	
+	// Don't run any of the deferred commands if the registry is still loading otherwise not all possible tests will be found and ran.
+	if (!AssetRegistryModule.Get().IsLoadingAssets())
+	{
+		// Execute all currently queued deferred commands (allows commands to be queued up for next frame).
+		const int32 DeferredCommandsCount = DeferredAutomationCommands.Num();
+		for (int32 DeferredCommandsIndex = 0; DeferredCommandsIndex < DeferredCommandsCount; DeferredCommandsIndex++)
+		{
+			{
+				GEngine->Exec(NULL, *DeferredAutomationCommands[DeferredCommandsIndex], *GLog);
+			}
+		}
+
+		// Once all of the commands have executed then we remove them from the array.
+		DeferredAutomationCommands.RemoveAt(0, DeferredCommandsCount);
+	}
+}
 
 /**
  * Implements a local controller to run tests and spew results, mostly used by automated testing.
@@ -544,7 +583,6 @@ bool GenerateTestNamesFromCommandLine(const TCHAR* InTestCommands, TArray<FStrin
 	return OutTestNames.Num() > 0;
 }
 
-
 bool DirectAutomationCommand(const TCHAR* Cmd, FOutputDevice* Ar = GLog)
 {
 	bool bResult = false;
@@ -616,6 +654,8 @@ bool DirectAutomationCommand(const TCHAR* Cmd, FOutputDevice* Ar = GLog)
 }
 
 
+
+
 static class FAutomationTestCmd : private FSelfRegisteringExec
 {
 public:
diff --git a/Engine/Source/Runtime/AutomationWorker/Private/AutomationWorkerModule.h b/Engine/Source/Runtime/AutomationWorker/Private/AutomationWorkerModule.h
--- a/Engine/Source/Runtime/AutomationWorker/Private/AutomationWorkerModule.h
+++ b/Engine/Source/Runtime/AutomationWorker/Private/AutomationWorkerModule.h
@@ -118,4 +118,10 @@ private:
 
 	/** Delegate to fire when the test is complete */
 	FStopTestEvent StopTestEvent;
+
+	// Holds the automation command line arguments.
+	TArray<FString> DeferredAutomationCommands;
+
+	// Cycles through and executes the automation commands in the deferred array.
+	void RunDeferredAutomationCommands();
 };
diff --git a/Engine/Source/Runtime/Engine/Private/UnrealEngine.cpp b/Engine/Source/Runtime/Engine/Private/UnrealEngine.cpp
--- a/Engine/Source/Runtime/Engine/Private/UnrealEngine.cpp
+++ b/Engine/Source/Runtime/Engine/Private/UnrealEngine.cpp
@@ -813,13 +813,6 @@ void UEngine::Init(IEngineLoop* InEngineLoop)
 		}
 	}
 
-	// Optionally queue automation tests
-	FString AutomationCmds;
-	if (FParse::Value(FCommandLine::Get(), TEXT("AutomationTests="), AutomationCmds, false))
-	{
-		new(GEngine->DeferredCommands) FString(FString(TEXT("Automation CommandLineTests ")) + AutomationCmds);
-	}
-
 	// optionally set the vsync console variable
 	if( FParse::Param(FCommandLine::Get(), TEXT("vsync")) )
 	{
@@ -3344,7 +3337,7 @@ bool UEngine::HandleListTexturesCommand( const TCHAR* Cmd, FOutputDevice& Ar )
 	SortedTextures.Sort( FCompareFSortedTexture( bAlphaSort ) );
 
 	// Retrieve mapping from LOD group enum value to text representation.
-	TArray<FString> TextureGroupNames = UTextureLODSettings::GetTextureGroupNames();
+	TArray<FString> TextureGroupNames = FTextureLODSettings::GetTextureGroupNames();
 
 	// Display.
 	int32 TotalMaxSize		= 0;
@@ -10652,6 +10645,17 @@ bool UEngine::ShouldAbsorbAuthorityOnlyEvent()
 }
 
 
+UDeviceProfileManager* UEngine::GetDeviceProfileManager()
+{
+	if(DeviceProfileManager == NULL)
+	{
+		DeviceProfileManager = NewObject<UDeviceProfileManager>(GetTransientPackage(), TEXT("GlobalDeviceProfileManager"), RF_Public | RF_Transient);
+	}
+
+	return DeviceProfileManager;
+}
+
+
 bool UEngine::ShouldAbsorbCosmeticOnlyEvent()
 {
 	for (auto It = WorldList.CreateIterator(); It; ++It)
