The game mode can often be too quick with the asset registry.  So if the asset registry is not running when the automation worker initializes then we just go ahead, cue up the tests and run 'em.

diff --git a/Engine/Source/Runtime/AutomationWorker/Private/AutomationWorkerModule.cpp b/Engine/Source/Runtime/AutomationWorker/Private/AutomationWorkerModule.cpp
--- a/Engine/Source/Runtime/AutomationWorker/Private/AutomationWorkerModule.cpp
+++ b/Engine/Source/Runtime/AutomationWorker/Private/AutomationWorkerModule.cpp
@@ -144,11 +144,17 @@ void FAutomationWorkerModule::Initialize()
 		new(DeferredAutomationCommands) FString(FString(TEXT("Automation CommandLineTests ")) + AutomationCmds);
 	}
 
+	//If the asset registry is loading assets then we'll wait for it to stop before running our automation tests.
 	FAssetRegistryModule& AssetRegistryModule = FModuleManager::LoadModuleChecked<FAssetRegistryModule>("AssetRegistry");
 	if (AssetRegistryModule.Get().IsLoadingAssets())
 	{
 		AssetRegistryModule.Get().OnFilesLoaded().AddRaw(this, &FAutomationWorkerModule::RunDeferredAutomationCommands);
 	}
+	else
+	{
+		//If the registry is not loading then we'll just go ahead and run our tests.
+		RunDeferredAutomationCommands();
+	}
 #endif // !(UE_BUILD_SHIPPING)
 }
 
