#Automation - Unified two commandline automation pathways (one via commandlet, one in the "worker") into one.
#NOTE - "-AutomationTests=" has been deprecated
#NOTE - Supported Commands are now "List", "RunTests", RunAll", and "Quit"

#NOTE - Syntax is now -ExecCmds="Automation List;RunAll;Quit"

diff --git a/Engine/Source/Developer/AutomationController/Private/AutomationCommandline.cpp b/Engine/Source/Developer/AutomationController/Private/AutomationCommandline.cpp
new file mode 100644
--- /dev/null
+++ b/Engine/Source/Developer/AutomationController/Private/AutomationCommandline.cpp
@@ -0,0 +1,348 @@
+// Copyright 1998-2015 Epic Games, Inc. All Rights Reserved.
+
+#include "AutomationControllerPrivatePCH.h"
+
+
+/** States for running the automation process */
+namespace EAutomationTestState
+{
+	enum Type
+	{
+		Idle,				// Automation process is not running
+		FindWorkers,		// Find workers to run the tests
+		RequestTests,		// Find the tests that can be run on the workers
+		DoingRequestedWork,	// Do whatever was requested from the commandline
+		Complete			// The process is finished
+	};
+};
+
+namespace EAutomationCommand
+{
+	enum Type
+	{
+		ListAllTests,				//List all tests for the session
+		//RunSingleTest,			//Run one test specified by the commandline
+		RunCommandLineTests,		//Run only tests that are listed on the commandline
+		RunAll,						//Run all the tests that are supported
+		Quit						//quit the app when tests are done
+	};
+};
+
+
+class FAutomationExecCmd : private FSelfRegisteringExec
+{
+public:
+	void Init()
+	{
+		SessionID = FApp::GetSessionId();
+
+		// Set state to FindWorkers to kick off the testing process
+		AutomationTestState = EAutomationTestState::Idle;
+		DelayTimer = 0.0f;
+
+		// Load the automation controller
+		IAutomationControllerModule* AutomationControllerModule = &FModuleManager::LoadModuleChecked<IAutomationControllerModule>("AutomationController");
+		AutomationController = AutomationControllerModule->GetAutomationController();
+
+		AutomationController->Init();
+
+		const bool bSkipScreenshots = FParse::Param(FCommandLine::Get(), TEXT("NoScreenshots"));
+		const bool bFullSizeScreenshots = FParse::Param(FCommandLine::Get(), TEXT("FullSizeScreenshots"));
+		AutomationController->SetScreenshotsEnabled(!bSkipScreenshots);
+		AutomationController->SetUsingFullSizeScreenshots(bFullSizeScreenshots);
+
+		AutomationController->SetPrintResults(true);
+
+		// Register for the callback that tells us there are tests available
+		AutomationController->OnTestsRefreshed().BindRaw(this, &FAutomationExecCmd::HandleRefreshTestCallback);
+
+		TickHandler = FTicker::GetCoreTicker().AddTicker(FTickerDelegate::CreateRaw(this, &FAutomationExecCmd::Tick));
+
+		int32 NumTestLoops = 1;
+		FParse::Value(FCommandLine::Get(), TEXT("TestLoops="), NumTestLoops);
+		AutomationController->SetNumPasses(NumTestLoops);
+		TestCount = 0;
+	}
+
+
+	void Shutdown()
+	{
+		//IAutomationControllerModule* AutomationControllerModule = &FModuleManager::LoadModuleChecked<IAutomationControllerModule>("AutomationController");
+		//AutomationControllerModule->ShutdownModule();
+
+		AutomationController->SetPrintResults(false);
+
+		FTicker::GetCoreTicker().RemoveTicker(TickHandler);
+	}
+
+
+	bool IsTestingComplete()
+	{
+		// If the automation controller is no longer processing and we've reached the final stage of testing
+		if ((AutomationController->GetTestState() != EAutomationControllerModuleState::Running) && (AutomationTestState == EAutomationTestState::Complete) && (AutomationCommandQueue.Num() == 0))
+		{
+			// If an actual test was ran we then will let the user know how many of them were ran.
+			if (TestCount > 0)
+			{
+				OutputDevice->Logf(TEXT("...Automation Test Queue Empty %d tests performed."), TestCount);
+				TestCount = 0;
+			}
+			return true;
+		}
+		return false;
+	}
+	
+
+	void GenerateTestNamesFromCommandLine(const TArray<FString>& AllTestNames, TArray<FString>& OutTestNames)
+	{
+		OutTestNames.Empty();
+
+		//Split the test names up
+		TArray<FString> Filters;
+		StringCommand.ParseIntoArray(Filters, TEXT("+"), true);
+		//trim cruft from all entries
+		for (int32 FilterIndex = 0; FilterIndex < Filters.Num(); ++FilterIndex)
+		{
+			Filters[FilterIndex] = Filters[FilterIndex].Trim();
+		}
+
+		for (int32 TestIndex = 0; TestIndex < AllTestNames.Num(); ++TestIndex)
+		{
+			for (int FilterIndex = 0; FilterIndex < Filters.Num(); ++FilterIndex)
+			{
+				if (AllTestNames[TestIndex].Contains(Filters[FilterIndex]))
+				{
+					OutTestNames.Add(AllTestNames[TestIndex]);
+					TestCount++;
+					break;
+				}
+			}
+		}
+	}
+
+
+	void FindWorkers(float DeltaTime)
+	{
+		// Time to delay requesting workers - ensure they are up and running
+		const float DelayTime = 5.0f;
+
+		DelayTimer += DeltaTime;
+
+		if (DelayTimer > DelayTime)
+		{
+			// Request the workers
+			AutomationController->RequestAvailableWorkers(SessionID);
+			AutomationTestState = EAutomationTestState::RequestTests;
+		}
+	}
+
+
+	void HandleRefreshTestCallback()
+	{
+		TArray <FString> AllTestNames;
+
+		// We have found some workers
+		// Create a filter to add to the automation controller, otherwise we don't get any reports
+		AutomationController->SetFilter(MakeShareable(new AutomationFilterCollection()));
+		AutomationController->SetVisibleTestsEnabled(true);
+		AutomationController->GetEnabledTestNames(AllTestNames);
+
+		//assume we won't run any tests
+		bool bRunTests = false;
+
+		if (AutomationCommand == EAutomationCommand::ListAllTests)
+		{
+
+			//TArray<FAutomationTestInfo> TestInfo;
+			//FAutomationTestFramework::GetInstance().GetValidTestNames(TestInfo);
+			for (int TestIndex = 0; TestIndex < AllTestNames.Num(); ++TestIndex)
+			{
+				OutputDevice->Logf(TEXT("%s"), *AllTestNames[TestIndex]);
+			}
+			OutputDevice->Logf(TEXT("Found %i Automation Tests"), AllTestNames.Num());
+			// Set state to complete
+			AutomationTestState = EAutomationTestState::Complete;
+		}
+		else if (AutomationCommand == EAutomationCommand::RunCommandLineTests)
+		{
+			TArray <FString> FilteredTestNames;
+			GenerateTestNamesFromCommandLine(AllTestNames, FilteredTestNames);
+			if (FilteredTestNames.Num())
+			{
+				AutomationController->SetEnabledTests(FilteredTestNames);
+				bRunTests = true;
+			}
+			else
+			{
+				AutomationTestState = EAutomationTestState::Complete;
+			}
+		}
+		else if (AutomationCommand == EAutomationCommand::RunAll)
+		{
+			bRunTests = true;
+			TestCount = AllTestNames.Num();
+		}
+
+		if (bRunTests)
+		{
+			AutomationController->RunTests();
+
+			// Set state to monitoring to check for test completion
+			AutomationTestState = EAutomationTestState::DoingRequestedWork;
+		}
+	}
+
+
+	void MonitorTests()
+	{
+		if (AutomationController->GetTestState() != EAutomationControllerModuleState::Running)
+		{
+			// We have finished the testing, and results are available
+			AutomationTestState = EAutomationTestState::Complete;
+		}
+	}
+
+
+
+	bool Tick(float DeltaTime)
+	{
+		// Update the automation controller to keep it running
+		AutomationController->Tick();
+
+		// Update the automation process
+		switch (AutomationTestState)
+		{
+			case EAutomationTestState::FindWorkers:
+			{
+				FindWorkers(DeltaTime);
+				break;
+			}
+			case EAutomationTestState::RequestTests:
+			{
+				break;
+			}
+			case EAutomationTestState::DoingRequestedWork:
+			{
+				MonitorTests();
+				break;
+			}
+			case EAutomationTestState::Complete:
+			case EAutomationTestState::Idle:
+			default:
+			{
+				//pop next command
+				if (AutomationCommandQueue.Num())
+				{
+					AutomationCommand = AutomationCommandQueue[0];
+					AutomationCommandQueue.RemoveAt(0);
+					if (AutomationCommand == EAutomationCommand::Quit)
+					{
+						if (AutomationCommandQueue.IsValidIndex(0))
+						{
+							// Add Quit back to the end of the array.
+							AutomationCommandQueue.Add(EAutomationCommand::Quit);
+							break;
+						}
+					}
+					AutomationTestState = EAutomationTestState::FindWorkers;
+				}
+
+				// Only quit if Quit is the actual last element in the array.
+				if (AutomationCommand == EAutomationCommand::Quit)
+				{
+					GIsRequestingExit = true;
+				}
+				break;
+			}
+		}
+		return !IsTestingComplete();
+	}
+
+	
+	/** Console commands, see embeded usage statement **/
+	virtual bool Exec(UWorld*, const TCHAR* Cmd, FOutputDevice& Ar) override
+	{
+		bool bHandled = false;
+
+		//save off device to send results to
+		OutputDevice = &Ar;
+
+		StringCommand.Empty();
+		//figure out if we are handling this request
+		if (FParse::Command(&Cmd, TEXT("Automation")))
+		{
+			TArray<FString> CommandList;
+			StringCommand = Cmd;
+			StringCommand.ParseIntoArray(CommandList, TEXT(";"), true);
+
+			//assume we handle this
+			bHandled = true;
+
+			for (int CommandIndex = 0; CommandIndex < CommandList.Num(); ++CommandIndex)
+			{
+				const TCHAR* TempCmd = *CommandList[CommandIndex];
+				if (FParse::Command(&TempCmd, TEXT("List")))
+				{
+					AutomationCommandQueue.Add(EAutomationCommand::ListAllTests);
+				}
+				else if (FParse::Command(&TempCmd, TEXT("RunTests")))
+				{
+					//only one of these should be used
+					StringCommand = TempCmd;
+					AutomationCommandQueue.Add(EAutomationCommand::RunCommandLineTests);
+				}
+				else if (FParse::Command(&TempCmd, TEXT("RunAll")))
+				{
+					AutomationCommandQueue.Add(EAutomationCommand::RunAll);
+				}
+				else if (FParse::Command(&TempCmd, TEXT("Quit")))
+				{
+					AutomationCommandQueue.Add(EAutomationCommand::Quit);
+				}
+			}
+		}
+
+		//we have our orders, let's automate some testing
+		if (bHandled)
+		{
+			Init();
+		}
+
+		// Shutdown our service
+		return bHandled;
+	}
+
+private:
+	/** The automation controller running the tests */
+	IAutomationControllerManagerPtr AutomationController;
+
+	/** The current state of the automation process */
+	EAutomationTestState::Type AutomationTestState;
+
+	/** What work was requested */
+	TArray<EAutomationCommand::Type> AutomationCommandQueue;
+
+	/** What work was requested */
+	EAutomationCommand::Type AutomationCommand;
+
+	/** Delay used before finding workers on game instances. Just to ensure they have started up */
+	float DelayTimer;
+
+	/** Holds the session ID */
+	FGuid SessionID;
+
+	//device to send results to
+	FOutputDevice* OutputDevice;
+
+	//so we can release control of the app and just get ticked like all other systems
+	FDelegateHandle TickHandler;
+
+	//Extra commandline params
+	FString StringCommand;
+
+	//This is the numbers of tests that are found in the command line.
+	int32 TestCount;
+};
+
+static FAutomationExecCmd AutomationExecCmd;
+
diff --git a/Engine/Source/Developer/AutomationController/Private/AutomationControllerManager.h b/Engine/Source/Developer/AutomationController/Private/AutomationControllerManager.h
--- a/Engine/Source/Developer/AutomationController/Private/AutomationControllerManager.h
+++ b/Engine/Source/Developer/AutomationController/Private/AutomationControllerManager.h
@@ -53,6 +53,12 @@ public:
 		bScreenshotsEnabled = bNewValue;
 	}
 
+	virtual void SetPrintResults(const bool bNewValue) override
+	{
+		bPrintResults = bNewValue;
+	}
+
+
 	virtual void SetFilter( TSharedPtr< AutomationFilterCollection > InFilter ) override
 	{
 		ReportManager.SetFilter( InFilter );
@@ -358,6 +364,9 @@ private:
 	/** If we should request full screen screen shots. */
 	bool bRequestFullScreenScreenshots;
 
+	/** If we want to echo results to the console */
+	bool bPrintResults;
+
 	/** If we should track any test history for the next run. */
 	bool bTrackHistory;
 
diff --git a/Engine/Source/Developer/AutomationController/Private/AutomationControllerManger.cpp b/Engine/Source/Developer/AutomationController/Private/AutomationControllerManger.cpp
--- a/Engine/Source/Developer/AutomationController/Private/AutomationControllerManger.cpp
+++ b/Engine/Source/Developer/AutomationController/Private/AutomationControllerManger.cpp
@@ -127,6 +127,7 @@ void FAutomationControllerManager::Init()
 	bTestResultsAvailable = false;
 	bScreenshotsEnabled = true;
 	bRequestFullScreenScreenshots = false;
+	bPrintResults = false;
 
 	// Update the ini with the settings
 	bTrackHistory = false;
@@ -223,7 +224,10 @@ void FAutomationControllerManager::ExecuteNextTask( int32 ClusterIndex, OUT bool
 						for (int32 AddressIndex = 0; AddressIndex < DeviceAddresses.Num(); ++AddressIndex)
 						{
 							FAutomationTestResults TestResults;
+
+							GLog->Logf(ELogVerbosity::Display, TEXT("Running Automation: '%s' (Class Name: '%s')"), *TestsRunThisPass[AddressIndex]->GetDisplayName(), *TestsRunThisPass[AddressIndex]->GetCommand());
 							TestResults.State = EAutomationState::InProcess;
+
 							TestResults.GameInstance = DeviceClusterManager.GetClusterDeviceName( ClusterIndex, DeviceIndex );
 							NextTest->SetResults( ClusterIndex,CurrentTestPass, TestResults );
 							NextTest->ResetNetworkCommandResponses();
@@ -630,7 +634,7 @@ void FAutomationControllerManager::HandleRequestTestsReplyMessage( const FAutoma
 
 	FAutomationTestInfo NewTest(Message.TestInfo);
 	TestInfo.Add(NewTest);
-
+	
 	if (TestInfo.Num() == NumOfTestsToReceive)
 	{
 		SetTestNames(Context->GetSender());
@@ -665,6 +669,30 @@ void FAutomationControllerManager::HandleRunTestsReplyMessage( const FAutomation
 
 		Report->SetResults(ClusterIndex,CurrentTestPass, TestResults);
 
+		if (bPrintResults)
+		{
+			for (TArray<FString>::TConstIterator ErrorIter(Message.Errors); ErrorIter; ++ErrorIter)
+			{
+				GLog->Logf(ELogVerbosity::Error, TEXT("%s"), **ErrorIter);
+			}
+			for (TArray<FString>::TConstIterator WarningIter(Message.Warnings); WarningIter; ++WarningIter)
+			{
+				GLog->Logf(ELogVerbosity::Warning, TEXT("%s"), **WarningIter);
+			}
+			for (TArray<FString>::TConstIterator LogItemIter(Message.Logs); LogItemIter; ++LogItemIter)
+			{
+				GLog->Logf(ELogVerbosity::Log, TEXT("%s"), **LogItemIter);
+			}
+			if (TestResults.State == EAutomationState::Success)
+			{
+				GLog->Logf(ELogVerbosity::Log, TEXT("...Automation Test Succeeded (%s)"), *Report->GetDisplayName());
+			}
+			else
+			{
+				GLog->Logf(ELogVerbosity::Log, TEXT("...Automation Test Failed (%s)"), *Report->GetDisplayName());
+			}
+		}
+
 		// Device is now good to go
 		DeviceClusterManager.SetTest(ClusterIndex, DeviceIndex, NULL);
 	}
diff --git a/Engine/Source/Developer/AutomationController/Public/Interfaces/IAutomationControllerManager.h b/Engine/Source/Developer/AutomationController/Public/Interfaces/IAutomationControllerManager.h
--- a/Engine/Source/Developer/AutomationController/Public/Interfaces/IAutomationControllerManager.h
+++ b/Engine/Source/Developer/AutomationController/Public/Interfaces/IAutomationControllerManager.h
@@ -222,6 +222,11 @@ public:
 	 */
 	virtual void SetScreenshotsEnabled( const bool bNewValue ) = 0;
 
+	/**
+	 * Sets if the results of tests should be echoed to the console
+	 */
+	virtual void SetPrintResults(const bool bNewValue) = 0;
+
 	/**
 	 * Filters the visible list of tests.
 	 */
diff --git a/Engine/Source/Editor/UnrealEd/Classes/Commandlets/AutomatedLauncherCommandlet.h b/Engine/Source/Editor/UnrealEd/Classes/Commandlets/AutomatedLauncherCommandlet.h
deleted file mode 100644
diff --git a/Engine/Source/Editor/UnrealEd/Private/Commandlets/AutomatedLauncherCommandlet.cpp b/Engine/Source/Editor/UnrealEd/Private/Commandlets/AutomatedLauncherCommandlet.cpp
deleted file mode 100644
diff --git a/Engine/Source/Editor/UnrealEd/UnrealEd.Build.cs b/Engine/Source/Editor/UnrealEd/UnrealEd.Build.cs
--- a/Engine/Source/Editor/UnrealEd/UnrealEd.Build.cs
+++ b/Engine/Source/Editor/UnrealEd/UnrealEd.Build.cs
@@ -99,7 +99,6 @@ public class UnrealEd : ModuleRules
 				"InputCore",
 				"InputBindingEditor",
 				"Internationalization",
-				"LauncherAutomatedService",
 				"LauncherServices",
 				"MaterialEditor",
 				"MessageLog",
diff --git a/Engine/Source/Runtime/AutomationWorker/Private/AutomationWorkerModule.cpp b/Engine/Source/Runtime/AutomationWorker/Private/AutomationWorkerModule.cpp
--- a/Engine/Source/Runtime/AutomationWorker/Private/AutomationWorkerModule.cpp
+++ b/Engine/Source/Runtime/AutomationWorker/Private/AutomationWorkerModule.cpp
@@ -134,27 +134,6 @@ void FAutomationWorkerModule::Initialize()
 	}
 	ExecutionCount = INDEX_NONE;
 	bExecutingNetworkCommandResults = false;
-
-#if !(UE_BUILD_SHIPPING)
-	//Obtain any command line tests commands that use '-automationtests='.
-	FString AutomationCmds;
-	if (FParse::Value(FCommandLine::Get(), TEXT("AutomationTests="), AutomationCmds, false))
-	{
-		new(DeferredAutomationCommands) FString(FString(TEXT("Automation CommandLineTests ")) + AutomationCmds);
-	}
-
-	//If the asset registry is loading assets then we'll wait for it to stop before running our automation tests.
-	FAssetRegistryModule& AssetRegistryModule = FModuleManager::LoadModuleChecked<FAssetRegistryModule>("AssetRegistry");
-	if (AssetRegistryModule.Get().IsLoadingAssets())
-	{
-		AssetRegistryModule.Get().OnFilesLoaded().AddRaw(this, &FAutomationWorkerModule::RunDeferredAutomationCommands);
-	}
-	else
-	{
-		//If the registry is not loading then we'll just go ahead and run our tests.
-		RunDeferredAutomationCommands();
-	}
-#endif // !(UE_BUILD_SHIPPING)
 }
 
 void FAutomationWorkerModule::ReportNetworkCommandComplete()
@@ -260,38 +239,58 @@ void FAutomationWorkerModule::SendTests( const FMessageAddress& ControllerAddres
 /* FAutomationWorkerModule callbacks
  *****************************************************************************/
 
-void FAutomationWorkerModule::HandleFindWorkersMessage( const FAutomationWorkerFindWorkers& Message, const IMessageContextRef& Context )
+void FAutomationWorkerModule::HandleFindWorkersMessage(const FAutomationWorkerFindWorkers& Message, const IMessageContextRef& Context)
 {
 	// Set the Instance name to be the same as the session browser. This information should be shared at some point
-	FString InstanceName = FString::Printf(TEXT("%s-%i"), FPlatformProcess::ComputerName(), FPlatformProcess::GetCurrentProcessId());
-
 	if ((Message.SessionId == FApp::GetSessionId()) && (Message.Changelist == 10000))
 	{
-		FAutomationWorkerFindWorkersResponse* Response = new FAutomationWorkerFindWorkersResponse();
-
-		FString OSMajorVersionString, OSSubVersionString;
-		FPlatformMisc::GetOSVersions( OSMajorVersionString, OSSubVersionString );
-
-		FString OSVersionString = OSMajorVersionString + TEXT(" ") + OSSubVersionString;
-		FString CPUModelString = FPlatformMisc::GetCPUBrand().Trim();
-
-		Response->DeviceName = FPlatformProcess::ComputerName();
-		Response->InstanceName = InstanceName;
-		Response->Platform = FPlatformProperties::PlatformName();
-		Response->SessionId = Message.SessionId;
-		Response->OSVersionName = OSVersionString;
-		Response->ModelName = FPlatformMisc::GetDefaultDeviceProfileName();
-		Response->GPUName = FPlatformMisc::GetPrimaryGPUBrand();
-		Response->CPUModelName = CPUModelString;
-		Response->RAMInGB = FPlatformMemory::GetPhysicalGBRam();
+		TestRequesterAddress = Context->GetSender();
+		//If the asset registry is loading assets then we'll wait for it to stop before running our automation tests.
+		FAssetRegistryModule& AssetRegistryModule = FModuleManager::LoadModuleChecked<FAssetRegistryModule>("AssetRegistry");
+		if (AssetRegistryModule.Get().IsLoadingAssets())
+		{
+			if (!AssetRegistryModule.Get().OnFilesLoaded().IsBoundToObject(this))
+			{
+				AssetRegistryModule.Get().OnFilesLoaded().AddRaw(this, &FAutomationWorkerModule::SendWorkerFound);
+				GLog->Logf(ELogVerbosity::Log, TEXT("...Forcing Asset Registry Load For Automation"));
+			}
+		}
+		else
+		{
+			//If the registry is not loading then we'll just go ahead and run our tests.
+			SendWorkerFound();
+		}
+	}
+}
+
+
+void FAutomationWorkerModule::SendWorkerFound()
+{
+	FAutomationWorkerFindWorkersResponse* Response = new FAutomationWorkerFindWorkersResponse();
+
+	FString OSMajorVersionString, OSSubVersionString;
+	FPlatformMisc::GetOSVersions(OSMajorVersionString, OSSubVersionString);
+
+	FString OSVersionString = OSMajorVersionString + TEXT(" ") + OSSubVersionString;
+	FString CPUModelString = FPlatformMisc::GetCPUBrand().Trim();
+
+	Response->DeviceName = FPlatformProcess::ComputerName();
+	Response->InstanceName = FString::Printf(TEXT("%s-%i"), FPlatformProcess::ComputerName(), FPlatformProcess::GetCurrentProcessId());
+	Response->Platform = FPlatformProperties::PlatformName();
+	Response->SessionId = FApp::GetSessionId();
+	Response->OSVersionName = OSVersionString;
+	Response->ModelName = FPlatformMisc::GetDefaultDeviceProfileName();
+	Response->GPUName = FPlatformMisc::GetPrimaryGPUBrand();
+	Response->CPUModelName = CPUModelString;
+	Response->RAMInGB = FPlatformMemory::GetPhysicalGBRam();
 #if WITH_ENGINE
-		Response->RenderModeName = AutomationCommon::GetRenderDetailsString();
+	Response->RenderModeName = AutomationCommon::GetRenderDetailsString();
 #else
-		Response->RenderModeName = TEXT("Unknown");
+	Response->RenderModeName = TEXT("Unknown");
 #endif
 
-		MessageEndpoint->Send(Response, Context->GetSender());
-	}
+	MessageEndpoint->Send(Response, TestRequesterAddress);
+	TestRequesterAddress.Invalidate();
 }
 
 
@@ -435,274 +434,4 @@ void FAutomationWorkerModule::HandleRunTestsMessage( const FAutomationWorkerRunT
 }
 
 
-void FAutomationWorkerModule::RunTest(const FString& InTestToRun, const int32 InRoleIndex, FStopTestEvent const& InStopTestEvent)
-{
-	TestName = InTestToRun;
-
-	StopTestEvent = InStopTestEvent;
-	// Always allow the first network command to execute
-	bExecuteNextNetworkCommand = true;
-
-	// We are not executing network command sub-commands right now
-	bExecutingNetworkCommandResults = false;
-
-	FAutomationTestFramework::GetInstance().StartTestByName(InTestToRun, InRoleIndex);
-}
-
-void FAutomationWorkerModule::RunDeferredAutomationCommands()
-{
-	// Get the loaded asset registry module.
-	FAssetRegistryModule& AssetRegistryModule = FModuleManager::LoadModuleChecked<FAssetRegistryModule>(TEXT("AssetRegistry"));
-	
-	// Don't run any of the deferred commands if the registry is still loading otherwise not all possible tests will be found and ran.
-	if (!AssetRegistryModule.Get().IsLoadingAssets())
-	{
-		// Execute all currently queued deferred commands (allows commands to be queued up for next frame).
-		const int32 DeferredCommandsCount = DeferredAutomationCommands.Num();
-		for (int32 DeferredCommandsIndex = 0; DeferredCommandsIndex < DeferredCommandsCount; DeferredCommandsIndex++)
-		{
-			{
-				GEngine->Exec(NULL, *DeferredAutomationCommands[DeferredCommandsIndex], *GLog);
-			}
-		}
-
-		// Once all of the commands have executed then we remove them from the array.
-		DeferredAutomationCommands.RemoveAt(0, DeferredCommandsCount);
-	}
-}
-
-/**
- * Implements a local controller to run tests and spew results, mostly used by automated testing.
- */
-static struct FQueueTests
-{
-	struct FJob
-	{
-		FString Test;
-		int32 RoleIndex;
-		FOutputDevice* Ar;
-
-		FJob(FString const& InTest, int32 InRoleIndex, FOutputDevice* InAr)
-			: Test(InTest)
-			, RoleIndex(InRoleIndex)
-			, Ar(InAr)
-		{
-		}
-	};
-
-	int32 NumTestsRun;
-	bool bTestInPogress;
-	bool bTicking;
-	TArray<FJob> Queue;
-
-	FQueueTests()
-		: NumTestsRun(0)
-		, bTestInPogress(false)
-		, bTicking(false)
-	{ }
-
-	void NewTest(FString const& Command, int32 RoleIndex = 0, FOutputDevice* Ar = GLog)
-	{
-		new (Queue) FJob(Command, RoleIndex, Ar);
-		if (!bTicking)
-		{
-			FTicker::GetCoreTicker().AddTicker(FTickerDelegate::CreateRaw(this, &FQueueTests::TickQueueTests), 0.1f);
-			bTicking = true;
-		}
-	}
-
-	bool TickQueueTests(float DeltaTime)
-	{
-		check(bTicking);
-		if (!GIsAutomationTesting && !bTestInPogress && Queue.Num())
-		{
-			FJob& CurrentJob = Queue[0];
-			TArray<FAutomationTestInfo> TestInfo;
-			FAutomationTestFramework::GetInstance().GetValidTestNames( TestInfo );
-			bool bRanIt = false;
-			for ( int TestIndex = 0; TestIndex < TestInfo.Num(); ++TestIndex )
-			{
-				FString TestCommand = TestInfo[TestIndex].GetTestName();
-				if (TestCommand == CurrentJob.Test)
-				{
-					CurrentJob.Ar->Logf(TEXT("Running: %s"), *CurrentJob.Test);
-					IAutomationWorkerModule::FStopTestEvent Event;
-					Event.BindRaw(this, &FQueueTests::ConsoleCommandTestComplete, CurrentJob.Ar);
-					if (FModuleManager::Get().IsModuleLoaded(TEXT("AutomationWorker")))
-					{
-						FModuleManager::GetModuleChecked<IAutomationWorkerModule>("AutomationWorker").RunTest(CurrentJob.Test, CurrentJob.RoleIndex, Event);
-						bTestInPogress = true;
-						bRanIt = true;
-					}
-					break;
-				}
-			}
-			if (!bRanIt)
-			{
-				CurrentJob.Ar->Logf(TEXT("ERROR: Failed to find test %s"), *CurrentJob.Test);
-			}
-			Queue.RemoveAt(0);
-		}
-		bTicking = !!Queue.Num();
-		return bTicking;
-	}
-
-	void ConsoleCommandTestComplete(bool bSuccess, FString Test, FAutomationTestExecutionInfo const& Results, FOutputDevice* Ar)
-	{
-		for ( TArray<FString>::TConstIterator ErrorIter( Results.Errors ); ErrorIter; ++ErrorIter )
-		{
-			Ar->Logf(ELogVerbosity::Error, TEXT("%s"), **ErrorIter);
-		}
-		for ( TArray<FString>::TConstIterator WarningIter( Results.Warnings ); WarningIter; ++WarningIter )
-		{
-			Ar->Logf(ELogVerbosity::Warning, TEXT("%s"), **WarningIter );
-		}
-		for ( TArray<FString>::TConstIterator LogItemIter( Results.LogItems ); LogItemIter; ++LogItemIter )
-		{
-			Ar->Logf(ELogVerbosity::Log, TEXT("%s"), **LogItemIter );
-		}
-		if (bSuccess)
-		{
-			Ar->Logf(ELogVerbosity::Log, TEXT("...Automation Test Succeeded (%s)"), *Test);
-		}
-		else
-		{
-			Ar->Logf(ELogVerbosity::Log, TEXT("...Automation Test Failed (%s)"), *Test);
-		}
-		bTestInPogress = false;
-		NumTestsRun++;
-		if (!bTicking)
-		{
-			GLog->Logf(ELogVerbosity::Log, TEXT("...Automation Test Queue Empty %d tests performed."), NumTestsRun);
-			NumTestsRun = 0;
-
-			//Quit when we are done if we have the command line option
-			const FString CommandLine(FCommandLine::Get());
-			if (CommandLine.Contains(TEXT("-AutomationTests=")))
-			{
-				//Let the current callstack complete then quit
-				if (GEngine->IsEditor())
-				{
-					GEngine->DeferredCommands.Add(TEXT("QUIT_EDITOR"));
-				}
-				else
-				{
-					GEngine->DeferredCommands.Add(TEXT("QUIT"));
-				}
-			}
-		}
-	}
-
-} QueueTests;
-
-bool GenerateTestNamesFromCommandLine(const TCHAR* InTestCommands, TArray<FString>& OutTestNames)
-{
-	//Split the test names up
-	TArray<FString> TestList;
-	const FString StringCommand = InTestCommands;
-	StringCommand.ParseIntoArray(TestList, TEXT(","), true);
-
-	//Get the list of valid names
-	TArray<FAutomationTestInfo> TestInfo;
-	FAutomationTestFramework::GetInstance().GetValidTestNames(TestInfo);
-
-	for (int32 TestListIndex = 0; TestListIndex < TestList.Num(); ++TestListIndex)
-	{
-		const FString CleanTestName = TestList[TestListIndex].Trim();
-		for (int TestIndex = 0; TestIndex < TestInfo.Num(); ++TestIndex)
-		{
-			if (TestInfo[TestIndex].GetDisplayName().StartsWith(CleanTestName))
-			{
-				OutTestNames.Add(TestInfo[TestIndex].GetTestName());
-			}
-		}
-	}
-
-	return OutTestNames.Num() > 0;
-}
-
-bool DirectAutomationCommand(const TCHAR* Cmd, FOutputDevice* Ar = GLog)
-{
-	bool bResult = false;
-	if(FParse::Command(&Cmd,TEXT("automation")))
-	{
-		const TCHAR* TempCmd = Cmd;
-		bResult = true;
-		if( FParse::Command( &TempCmd, TEXT( "list" ) ) )
-		{
-			TArray<FAutomationTestInfo> TestInfo;
-			FAutomationTestFramework::GetInstance().GetValidTestNames( TestInfo );
-			for ( int TestIndex = 0; TestIndex < TestInfo.Num(); ++TestIndex )
-			{
-				FString TestCommand = TestInfo[TestIndex].GetTestName();
-				Ar->Logf(TEXT("%s"), *TestCommand);
-			}
-		}
-		else if( FParse::Command(&TempCmd,TEXT("run")) )
-		{
-			FString Test(TempCmd);
-			Test = Test.Trim();
-			QueueTests.NewTest(Test);
-		}
-		else if (FParse::Command(&TempCmd, TEXT("CommandLineTests")))
-		{
-			TArray<FString> TestNames;
-
-			const bool bSkipScreenshots = FParse::Param(FCommandLine::Get(), TEXT("NoScreenshots"));
-			const bool bFullSizeScreenshots = FParse::Param(FCommandLine::Get(), TEXT("FullSizeScreenshots"));
-			FAutomationTestFramework::GetInstance().SetScreenshotOptions(!bSkipScreenshots, bFullSizeScreenshots);
-
-			if (GenerateTestNamesFromCommandLine(TempCmd, TestNames))
-			{
-				//Get the number of times to loop
-				int32 NumTestLoops = 1;
-				FParse::Value(FCommandLine::Get(), TEXT("TestLoops="), NumTestLoops);
-
-				for (int32 LoopIndex = 0; LoopIndex < NumTestLoops; ++LoopIndex)
-				{
-					for (int32 TestIndex = 0; TestIndex < TestNames.Num(); ++TestIndex)
-					{
-						QueueTests.NewTest(TestNames[TestIndex]);
-					}
-				}
-			}
-		}
-		else if( FParse::Command(&TempCmd,TEXT("runall")) )
-		{
-			int32 Mod = 0;
-			int32 Rem = 0;
-			FParse::Value(Cmd, TEXT("MOD="), Mod);
-			FParse::Value(Cmd, TEXT("REM="), Rem);
-			TArray<FAutomationTestInfo> TestInfo;
-			FAutomationTestFramework::GetInstance().GetValidTestNames( TestInfo );
-			for ( int TestIndex = 0; TestIndex < TestInfo.Num(); ++TestIndex )
-			{
-				if (!Mod || TestIndex % Mod == Rem)
-				{
-					QueueTests.NewTest(TestInfo[TestIndex].GetTestName());
-				}
-			}
-		}
-		else
-		{
-			bResult = false;
-		}
-	}
-	return bResult;
-}
-
-
-
-
-static class FAutomationTestCmd : private FSelfRegisteringExec
-{
-public:
-	/** Console commands, see embeded usage statement **/
-	virtual bool Exec( UWorld*, const TCHAR* Cmd, FOutputDevice& Ar) override
-	{
-		return DirectAutomationCommand(Cmd, &Ar);
-	}
-} AutomationTestCmd;
-
-
 #undef LOCTEXT_NAMESPACE
diff --git a/Engine/Source/Runtime/AutomationWorker/Private/AutomationWorkerModule.h b/Engine/Source/Runtime/AutomationWorker/Private/AutomationWorkerModule.h
--- a/Engine/Source/Runtime/AutomationWorker/Private/AutomationWorkerModule.h
+++ b/Engine/Source/Runtime/AutomationWorker/Private/AutomationWorkerModule.h
@@ -26,8 +26,6 @@ public:
 
 	virtual void Tick() override;
 
-	virtual void RunTest( const FString& InTestToRun, const int32 InRoleIndex, FStopTestEvent const& InStopTestEvent ) override;
-
 	// End IAutomationWorkerModule interface
 
 protected:
@@ -69,6 +67,9 @@ private:
 	// Handles FAutomationWorkerFindWorkers messages.
 	void HandleFindWorkersMessage( const FAutomationWorkerFindWorkers& Message, const IMessageContextRef& Context );
 
+	//deferred handler for sending "find worker" response in case the asset registry isn't loaded yet
+	void SendWorkerFound();
+
 	// Handles message endpoint shutdowns.
 	void HandleMessageEndpointShutdown();
 
@@ -129,7 +130,4 @@ private:
 
 	// Holds the automation command line arguments.
 	TArray<FString> DeferredAutomationCommands;
-
-	// Cycles through and executes the automation commands in the deferred array.
-	void RunDeferredAutomationCommands();
 };
diff --git a/Engine/Source/Runtime/AutomationWorker/Public/Interfaces/IAutomationWorkerModule.h b/Engine/Source/Runtime/AutomationWorker/Public/Interfaces/IAutomationWorkerModule.h
--- a/Engine/Source/Runtime/AutomationWorker/Public/Interfaces/IAutomationWorkerModule.h
+++ b/Engine/Source/Runtime/AutomationWorker/Public/Interfaces/IAutomationWorkerModule.h
@@ -23,7 +23,7 @@ public:
 	 * @param	InRoleIndex			Identifier for which worker in this group that should execute a command
 	 * @param	InStopTestEvent		Delegate to fire when the command is finished
 	 */
-	virtual void RunTest(const FString& InTestToRun, const int32 InRoleIndex, FStopTestEvent const& InStopTestEvent) = 0;
+	//virtual void RunTest(const FString& InTestToRun, const int32 InRoleIndex, FStopTestEvent const& InStopTestEvent) = 0;
 
 
 protected:
diff --git a/Engine/Source/Runtime/Engine/Private/Tests/EngineAutomationTests.cpp b/Engine/Source/Runtime/Engine/Private/Tests/EngineAutomationTests.cpp
--- a/Engine/Source/Runtime/Engine/Private/Tests/EngineAutomationTests.cpp
+++ b/Engine/Source/Runtime/Engine/Private/Tests/EngineAutomationTests.cpp
@@ -356,7 +356,7 @@ bool FCinematicFPSPerfTest::RunTest(const FString& Parameters)
 
 	//Check we are running from commandline
 	const FString CommandLine(FCommandLine::Get());
-	if (CommandLine.Contains(TEXT("AutomationTests")))
+	if (CommandLine.Contains(TEXT("Automation")))
 	{
 		//Get the name of the matinee to be used.
 		//If the game was not launched with the -MatineeName argument then this test will be ran based on time.
diff --git a/Engine/Source/Runtime/Launch/Private/LaunchEngineLoop.cpp b/Engine/Source/Runtime/Launch/Private/LaunchEngineLoop.cpp
--- a/Engine/Source/Runtime/Launch/Private/LaunchEngineLoop.cpp
+++ b/Engine/Source/Runtime/Launch/Private/LaunchEngineLoop.cpp
@@ -2056,12 +2056,13 @@ int32 FEngineLoop::Init()
 #endif
 
 #if WITH_EDITOR
-	if( GIsEditor )
+	FModuleManager::Get().LoadModule("AutomationController");
+	FModuleManager::GetModuleChecked<IAutomationControllerModule>("AutomationController").Init();
+	if (GIsEditor)
 	{
-		FModuleManager::GetModuleChecked<IAutomationControllerModule>("AutomationController").Init();
 		FModuleManager::Get().LoadModule(TEXT("ProfilerClient"));
 	}
-#endif
+#endif#endif
 
 	GIsRunning = true;
 
@@ -2411,17 +2412,14 @@ void FEngineLoop::Tick()
 #endif
 
 #if WITH_EDITOR
-		if( GIsEditor )
-		{
 #if WITH_ENGINE
-			QUICK_SCOPE_CYCLE_COUNTER(STAT_FEngineLoop_Tick_AutomationController);
+		QUICK_SCOPE_CYCLE_COUNTER(STAT_FEngineLoop_Tick_AutomationController);
 #endif
-			static FName AutomationController("AutomationController");
-			//Check if module loaded to support the change to allow this to be hot compilable.
-			if (FModuleManager::Get().IsModuleLoaded(AutomationController))
-			{
-				FModuleManager::GetModuleChecked<IAutomationControllerModule>(AutomationController).Tick();
-			}
+		static FName AutomationController("AutomationController");
+		//Check if module loaded to support the change to allow this to be hot compilable.
+		if (FModuleManager::Get().IsModuleLoaded(AutomationController))
+		{
+			FModuleManager::GetModuleChecked<IAutomationControllerModule>(AutomationController).Tick();
 		}
 #endif
 
diff --git a/Engine/Source/Runtime/LauncherAutomatedService/LauncherAutomatedService.Build.cs b/Engine/Source/Runtime/LauncherAutomatedService/LauncherAutomatedService.Build.cs
deleted file mode 100644
diff --git a/Engine/Source/Runtime/LauncherAutomatedService/Private/AutomatedTestManager.cpp b/Engine/Source/Runtime/LauncherAutomatedService/Private/AutomatedTestManager.cpp
deleted file mode 100644
diff --git a/Engine/Source/Runtime/LauncherAutomatedService/Private/AutomatedTestManager.h b/Engine/Source/Runtime/LauncherAutomatedService/Private/AutomatedTestManager.h
deleted file mode 100644
diff --git a/Engine/Source/Runtime/LauncherAutomatedService/Private/LauncherAutomatedServiceModule.cpp b/Engine/Source/Runtime/LauncherAutomatedService/Private/LauncherAutomatedServiceModule.cpp
deleted file mode 100644
diff --git a/Engine/Source/Runtime/LauncherAutomatedService/Private/LauncherAutomatedServiceModule.h b/Engine/Source/Runtime/LauncherAutomatedService/Private/LauncherAutomatedServiceModule.h
deleted file mode 100644
diff --git a/Engine/Source/Runtime/LauncherAutomatedService/Private/LauncherAutomatedServicePrivatePCH.h b/Engine/Source/Runtime/LauncherAutomatedService/Private/LauncherAutomatedServicePrivatePCH.h
deleted file mode 100644
diff --git a/Engine/Source/Runtime/LauncherAutomatedService/Private/LauncherAutomatedServiceProvider.cpp b/Engine/Source/Runtime/LauncherAutomatedService/Private/LauncherAutomatedServiceProvider.cpp
deleted file mode 100644
diff --git a/Engine/Source/Runtime/LauncherAutomatedService/Private/LauncherAutomatedServiceProvider.h b/Engine/Source/Runtime/LauncherAutomatedService/Private/LauncherAutomatedServiceProvider.h
deleted file mode 100644
diff --git a/Engine/Source/Runtime/LauncherAutomatedService/Public/Interfaces/IAutomatedTestManager.h b/Engine/Source/Runtime/LauncherAutomatedService/Public/Interfaces/IAutomatedTestManager.h
deleted file mode 100644
diff --git a/Engine/Source/Runtime/LauncherAutomatedService/Public/Interfaces/ILauncherAutomatedServiceModule.h b/Engine/Source/Runtime/LauncherAutomatedService/Public/Interfaces/ILauncherAutomatedServiceModule.h
deleted file mode 100644
diff --git a/Engine/Source/Runtime/LauncherAutomatedService/Public/Interfaces/ILauncherAutomatedServiceProvider.h b/Engine/Source/Runtime/LauncherAutomatedService/Public/Interfaces/ILauncherAutomatedServiceProvider.h
deleted file mode 100644
diff --git a/Engine/Source/Runtime/LauncherAutomatedService/Public/LauncherAutomatedService.h b/Engine/Source/Runtime/LauncherAutomatedService/Public/LauncherAutomatedService.h
deleted file mode 100644
