The automation worker now only checks the asset registry in editor builds.

Merging using UE4-To-UE4-LauncherDev

diff --git a/Engine/Source/Runtime/AutomationWorker/Private/AutomationWorkerModule.cpp b/Engine/Source/Runtime/AutomationWorker/Private/AutomationWorkerModule.cpp
--- a/Engine/Source/Runtime/AutomationWorker/Private/AutomationWorkerModule.cpp
+++ b/Engine/Source/Runtime/AutomationWorker/Private/AutomationWorkerModule.cpp
@@ -1,8 +1,10 @@
 // Copyright 1998-2015 Epic Games, Inc. All Rights Reserved.
 
 #include "AutomationWorkerPrivatePCH.h"
-#include "AssetRegistryModule.h"
 
+#if WITH_EDITOR
+#include "AssetRegistryModule.h"
+#endif
 
 #define LOCTEXT_NAMESPACE "AutomationTest"
 
@@ -245,6 +247,8 @@ void FAutomationWorkerModule::HandleFindWorkersMessage(const FAutomationWorkerFi
 	if ((Message.SessionId == FApp::GetSessionId()) && (Message.Changelist == 10000))
 	{
 		TestRequesterAddress = Context->GetSender();
+
+#if WITH_EDITOR
 		//If the asset registry is loading assets then we'll wait for it to stop before running our automation tests.
 		FAssetRegistryModule& AssetRegistryModule = FModuleManager::LoadModuleChecked<FAssetRegistryModule>("AssetRegistry");
 		if (AssetRegistryModule.Get().IsLoadingAssets())
@@ -256,6 +260,7 @@ void FAutomationWorkerModule::HandleFindWorkersMessage(const FAutomationWorkerFi
 			}
 		}
 		else
+#endif
 		{
 			//If the registry is not loading then we'll just go ahead and run our tests.
 			SendWorkerFound();
