Build fixes

diff --git a/Engine/Build/InstalledEngineBuild.xml b/Engine/Build/InstalledEngineBuild.xml
--- a/Engine/Build/InstalledEngineBuild.xml
+++ b/Engine/Build/InstalledEngineBuild.xml
@@ -254,7 +254,7 @@
 		</Node>
 	
 		<!-- Downstream monolithics Switch -->
-		<Node Name="Compile UE4Game Switch" Requires="Compile UnrealHeaderTool Win64" Produces="#UE4Game Switch;#UE4Game Switch Includes">
+		<Node Name="Compile UE4Game Switch" Requires="Compile UnrealHeaderTool Win64" Produces="#UE4Game Switch;#UE4Game Switch Includes" If="'$(WithSwitch)' == true">
 			<Compile Target="UE4Game" Platform="Switch" Configuration="Development" Tag="#UE4Game Switch" Arguments="-precompile"/>
 			<Compile Target="UE4Game" Platform="Switch" Configuration="Shipping" Tag="#UE4Game Switch" Arguments="-precompile"/>
 
diff --git a/Engine/Plugins/Online/OnlineSubsystemFacebook/Source/Private/IOS/OnlineFriendsFacebook.cpp b/Engine/Plugins/Online/OnlineSubsystemFacebook/Source/Private/IOS/OnlineFriendsFacebook.cpp
--- a/Engine/Plugins/Online/OnlineSubsystemFacebook/Source/Private/IOS/OnlineFriendsFacebook.cpp
+++ b/Engine/Plugins/Online/OnlineSubsystemFacebook/Source/Private/IOS/OnlineFriendsFacebook.cpp
@@ -5,6 +5,8 @@
 #include "OnlineSubsystemFacebookPrivate.h"
 #include "OnlineIdentityFacebook.h"
 #include "OnlineSharingFacebook.h"
+#include "Interfaces/OnlineFriendsInterface.h"
+#include "Misc/ConfigCacheIni.h"
 
 #import <FBSDKCoreKit/FBSDKCoreKit.h>
 
diff --git a/Engine/Plugins/Online/OnlineSubsystemFacebook/Source/Private/IOS/OnlineFriendsFacebook.h b/Engine/Plugins/Online/OnlineSubsystemFacebook/Source/Private/IOS/OnlineFriendsFacebook.h
--- a/Engine/Plugins/Online/OnlineSubsystemFacebook/Source/Private/IOS/OnlineFriendsFacebook.h
+++ b/Engine/Plugins/Online/OnlineSubsystemFacebook/Source/Private/IOS/OnlineFriendsFacebook.h
@@ -2,10 +2,13 @@
 
 #pragma once
 
-#include "OnlineFriendsInterface.h"
+#include "CoreMinimal.h"
+#include "OnlineSubsystem.h"
+#include "Interfaces/OnlineFriendsInterface.h"
+#include "Interfaces/OnlineSharingInterface.h"
+#include "Interfaces/OnlinePresenceInterface.h"
 #include "OnlineSharingFacebook.h"
 #include "OnlineSubsystemFacebookPackage.h"
-#include "OnlinePresenceInterface.h"
 
 /**
  * Info associated with an online friend on the facebook service
diff --git a/Engine/Plugins/Online/OnlineSubsystemFacebook/Source/Private/IOS/OnlineSharingFacebook.h b/Engine/Plugins/Online/OnlineSubsystemFacebook/Source/Private/IOS/OnlineSharingFacebook.h
--- a/Engine/Plugins/Online/OnlineSubsystemFacebook/Source/Private/IOS/OnlineSharingFacebook.h
+++ b/Engine/Plugins/Online/OnlineSubsystemFacebook/Source/Private/IOS/OnlineSharingFacebook.h
@@ -5,7 +5,9 @@
  
 
 // Module includes
-#include "OnlineSharingInterface.h"
+#include "CoreMinimal.h"
+#include "Interfaces/OnlineSharingInterface.h"
+#include "Interfaces/OnlineIdentityInterface.h"
 #include "OnlineSubsystemFacebookPackage.h"
 
 /**
diff --git a/Engine/Plugins/Online/OnlineSubsystemFacebook/Source/Private/IOS/OnlineSubsystemFacebook.cpp b/Engine/Plugins/Online/OnlineSubsystemFacebook/Source/Private/IOS/OnlineSubsystemFacebook.cpp
--- a/Engine/Plugins/Online/OnlineSubsystemFacebook/Source/Private/IOS/OnlineSubsystemFacebook.cpp
+++ b/Engine/Plugins/Online/OnlineSubsystemFacebook/Source/Private/IOS/OnlineSubsystemFacebook.cpp
@@ -1,7 +1,9 @@
 // Copyright 1998-2016 Epic Games, Inc. All Rights Reserved.
 
-#include "OnlineSubsystemFacebookPrivatePCH.h"
+#include "OnlineSubsystemFacebook.h"
+#include "OnlineSubsystemFacebookPrivate.h"
 #include "IOSAppDelegate.h"
+#include "Misc/CommandLine.h"
 
 #include "OnlineSharingFacebook.h"
 #include "OnlineUserFacebook.h"
diff --git a/Engine/Plugins/Runtime/Advertising/IOSTapJoy/Source/IOSTapJoy/Private/IOSTapJoy.cpp b/Engine/Plugins/Runtime/Advertising/IOSTapJoy/Source/IOSTapJoy/Private/IOSTapJoy.cpp
--- a/Engine/Plugins/Runtime/Advertising/IOSTapJoy/Source/IOSTapJoy/Private/IOSTapJoy.cpp
+++ b/Engine/Plugins/Runtime/Advertising/IOSTapJoy/Source/IOSTapJoy/Private/IOSTapJoy.cpp
@@ -1,7 +1,11 @@
 // Copyright 1998-2016 Epic Games, Inc. All Rights Reserved.
 
-#include "IOSTapJoyPrivatePCH.h"
+#include "CoreMinimal.h"
+#include "IAdvertisingProvider.h"
 #include "IOSAppDelegate.h"
+#include "Modules/ModuleInterface.h"
+#include "Modules/ModuleManager.h"
+#include "Misc/ConfigCacheIni.h"
 #import <Tapjoy/Tapjoy.h>
 
 class FTapJoyProvider : public IAdvertisingProvider
diff --git a/Engine/Plugins/Runtime/Analytics/Apsalar/Source/IOSApsalar/Private/IOSApsalar.cpp b/Engine/Plugins/Runtime/Analytics/Apsalar/Source/IOSApsalar/Private/IOSApsalar.cpp
--- a/Engine/Plugins/Runtime/Analytics/Apsalar/Source/IOSApsalar/Private/IOSApsalar.cpp
+++ b/Engine/Plugins/Runtime/Analytics/Apsalar/Source/IOSApsalar/Private/IOSApsalar.cpp
@@ -2,6 +2,7 @@
 
 #include "IOSApsalar.h"
 #include "IOSApsalarProvider.h"
+#include "Modules/ModuleManager.h"
 
 #if WITH_APSALAR
 #import "Apsalar.h"
diff --git a/Engine/Plugins/Runtime/Analytics/Apsalar/Source/IOSApsalar/Public/IOSApsalar.h b/Engine/Plugins/Runtime/Analytics/Apsalar/Source/IOSApsalar/Public/IOSApsalar.h
--- a/Engine/Plugins/Runtime/Analytics/Apsalar/Source/IOSApsalar/Public/IOSApsalar.h
+++ b/Engine/Plugins/Runtime/Analytics/Apsalar/Source/IOSApsalar/Public/IOSApsalar.h
@@ -2,7 +2,9 @@
 
 #pragma once
 
+#include "CoreMinimal.h"
 #include "IAnalyticsProviderModule.h"
+#include "Modules/ModuleManager.h"
 
 class IAnalyticsProvider;
 
diff --git a/Engine/Source/Developer/HotReload/Private/HotReloadClassReinstancer.cpp b/Engine/Source/Developer/HotReload/Private/HotReloadClassReinstancer.cpp
--- a/Engine/Source/Developer/HotReload/Private/HotReloadClassReinstancer.cpp
+++ b/Engine/Source/Developer/HotReload/Private/HotReloadClassReinstancer.cpp
@@ -1,13 +1,15 @@
 // Copyright 1998-2016 Epic Games, Inc. All Rights Reserved.
 
 #include "HotReloadClassReinstancer.h"
-#include "Engine/Blueprint.h"
 #include "Serialization/MemoryWriter.h"
 #include "UObject/UObjectHash.h"
 #include "UObject/UObjectIterator.h"
 #include "UObject/Package.h"
 #include "Serialization/ArchiveReplaceObjectRef.h"
+#if WITH_ENGINE
+#include "Engine/Blueprint.h"
 #include "Engine/BlueprintGeneratedClass.h"
+#endif
 
 #if WITH_ENGINE
 
diff --git a/Engine/Source/Developer/HotReload/Private/HotReloadClassReinstancer.h b/Engine/Source/Developer/HotReload/Private/HotReloadClassReinstancer.h
--- a/Engine/Source/Developer/HotReload/Private/HotReloadClassReinstancer.h
+++ b/Engine/Source/Developer/HotReload/Private/HotReloadClassReinstancer.h
@@ -2,7 +2,9 @@
 #pragma once
 
 #include "CoreMinimal.h"
+#if WITH_ENGINE
 #include "Kismet2/KismetReinstanceUtilities.h"
+#endif
 
 class UBlueprint;
 
diff --git a/Engine/Source/Developer/SuperSearch/Public/SuperSearchModule.h b/Engine/Source/Developer/SuperSearch/Public/SuperSearchModule.h
--- a/Engine/Source/Developer/SuperSearch/Public/SuperSearchModule.h
+++ b/Engine/Source/Developer/SuperSearch/Public/SuperSearchModule.h
@@ -5,7 +5,9 @@
 #include "CoreMinimal.h"
 #include "UObject/ObjectMacros.h"
 #include "Widgets/SWidget.h"
+#if WITH_EDITOR
 #include "AssetData.h"
+#endif
 #include "Modules/ModuleInterface.h"
 #include "Modules/ModuleManager.h"
 
diff --git a/Engine/Source/Programs/CrashReporter/MinidumpDiagnostics/Private/Mac/MinidumpDiagnosticsMainMac.cpp b/Engine/Source/Programs/CrashReporter/MinidumpDiagnostics/Private/Mac/MinidumpDiagnosticsMainMac.cpp
--- a/Engine/Source/Programs/CrashReporter/MinidumpDiagnostics/Private/Mac/MinidumpDiagnosticsMainMac.cpp
+++ b/Engine/Source/Programs/CrashReporter/MinidumpDiagnostics/Private/Mac/MinidumpDiagnosticsMainMac.cpp
@@ -1,8 +1,8 @@
 // Copyright 1998-2016 Epic Games, Inc. All Rights Reserved.
 
 #include "MinidumpDiagnosticsApp.h"
+#include "HAL/PlatformMisc.h"
 
-#include "Core.h"
 #include "RequiredProgramMainCPPInclude.h"
 #include "ExceptionHandling.h"
 
diff --git a/Engine/Source/Runtime/AutomationWorker/Private/AutomationWorkerModule.cpp b/Engine/Source/Runtime/AutomationWorker/Private/AutomationWorkerModule.cpp
--- a/Engine/Source/Runtime/AutomationWorker/Private/AutomationWorkerModule.cpp
+++ b/Engine/Source/Runtime/AutomationWorker/Private/AutomationWorkerModule.cpp
@@ -8,14 +8,16 @@
 #include "Modules/ModuleManager.h"
 #include "Helpers/MessageEndpointBuilder.h"
 #include "AutomationWorkerMessages.h"
-#include "UnrealClient.h"
-#include "Engine/GameViewportClient.h"
+#include "AutomationAnalytics.h"
+#include "JsonObjectConverter.h"
+
+#if WITH_ENGINE
 #include "ImageUtils.h"
-#include "EngineGlobals.h"
 #include "Engine/Engine.h"
+#include "Engine/GameViewportClient.h"
+#include "UnrealClient.h"
 #include "Tests/AutomationCommon.h"
-#include "AutomationAnalytics.h"
-#include "JsonObjectConverter.h"
+#endif
 
 #if WITH_EDITOR
 #include "AssetRegistryModule.h"
diff --git a/Engine/Source/Runtime/Core/Private/HAL/MallocLeakDetection.cpp b/Engine/Source/Runtime/Core/Private/HAL/MallocLeakDetection.cpp
--- a/Engine/Source/Runtime/Core/Private/HAL/MallocLeakDetection.cpp
+++ b/Engine/Source/Runtime/Core/Private/HAL/MallocLeakDetection.cpp
@@ -6,6 +6,16 @@
 
 #include "HAL/MallocLeakDetection.h"
 #include "Logging/LogMacros.h"
+#include "HAL/PlatformTLS.h"
+#include "HAL/IConsoleManager.h"
+#include "HAL/FileManager.h"
+#include "HAL/PlatformStackWalk.h"
+#include "Misc/Parse.h"
+#include "Misc/Paths.h"
+#include "Logging/LogMacros.h"
+#include "Misc/OutputDeviceArchiveWrapper.h"
+#include "Misc/DateTime.h"
+#include "ProfilingDebugging/ProfilingHelpers.h"
 
 #if MALLOC_LEAKDETECTION
 
diff --git a/Engine/Source/Runtime/Core/Public/Windows/WindowsPlatformMath.h b/Engine/Source/Runtime/Core/Public/Windows/WindowsPlatformMath.h
--- a/Engine/Source/Runtime/Core/Public/Windows/WindowsPlatformMath.h
+++ b/Engine/Source/Runtime/Core/Public/Windows/WindowsPlatformMath.h
@@ -128,7 +128,7 @@ struct FWindowsPlatformMath : public FGenericPlatformMath
 	static FORCEINLINE uint64 CountLeadingZeros64(uint64 Value)
 	{
 		// Use BSR to return the log2 of the integer
-		DWORD Log2;
+		unsigned long Log2;
 		if (_BitScanReverse64(&Log2, Value) != 0)
 		{
 			return 63 - Log2;
diff --git a/Engine/Source/Runtime/GameMenuBuilder/Private/GameMenuBuilder.cpp b/Engine/Source/Runtime/GameMenuBuilder/Private/GameMenuBuilder.cpp
deleted file mode 100644
diff --git a/Engine/Source/Runtime/GameMenuBuilder/Private/GameMenuBuilderModule.cpp b/Engine/Source/Runtime/GameMenuBuilder/Private/GameMenuBuilderModule.cpp
new file mode 100644
--- /dev/null
+++ b/Engine/Source/Runtime/GameMenuBuilder/Private/GameMenuBuilderModule.cpp
@@ -0,0 +1,10 @@
+// Copyright 1998-2016 Epic Games, Inc. All Rights Reserved.
+
+#include "GameMenuBuilderModule.h"
+#include "Modules/ModuleManager.h"
+
+
+
+IMPLEMENT_MODULE(FGameMenuBuilderModule, GameMenuBuilder)
+
+
diff --git a/Engine/Source/Runtime/GameMenuBuilder/Private/GameMenuBuilderStyle.cpp b/Engine/Source/Runtime/GameMenuBuilder/Private/GameMenuBuilderStyle.cpp
--- a/Engine/Source/Runtime/GameMenuBuilder/Private/GameMenuBuilderStyle.cpp
+++ b/Engine/Source/Runtime/GameMenuBuilder/Private/GameMenuBuilderStyle.cpp
@@ -6,7 +6,7 @@
 #include "Styling/SlateTypes.h"
 #include "Framework/Application/SlateApplication.h"
 #include "Slate/SlateGameResources.h"
-#include "GameMenuBuilder.h"
+#include "GameMenuBuilderModule.h"
 
 
 
@@ -16,7 +16,7 @@ void FGameMenuBuilderStyle::Initialize(const FString StyleName)
 {
 	if (FModuleManager::Get().IsModuleLoaded("GameMenuBuilder") == false)
 	{
-		FModuleManager::LoadModuleChecked<FGameMenuBuilderModule>("GameMenuBuilder");
+		FModuleManager::LoadModuleChecked<IGameMenuBuilderModule>("GameMenuBuilder");
 	}
 	if (!SimpleStyleInstance.IsValid())
 	{
diff --git a/Engine/Source/Runtime/GameMenuBuilder/Private/GameMenuPage.cpp b/Engine/Source/Runtime/GameMenuBuilder/Private/GameMenuPage.cpp
--- a/Engine/Source/Runtime/GameMenuBuilder/Private/GameMenuPage.cpp
+++ b/Engine/Source/Runtime/GameMenuBuilder/Private/GameMenuPage.cpp
@@ -6,7 +6,7 @@
 #include "GameMenuBuilderStyle.h"
 #include "GameMenuWidgetStyle.h"
 #include "SGameMenuPageWidget.h"
-#include "GameMenuBuilder.h"
+#include "GameMenuBuilderModule.h"
 
 FGameMenuPage::FGameMenuPage()
 {
@@ -34,7 +34,7 @@ bool FGameMenuPage::InitialiseRootMenu(APlayerController* InPCOwner, const FStri
 {
 	if (FModuleManager::Get().IsModuleLoaded("GameMenuBuilder") == false)
 	{
-		FModuleManager::LoadModuleChecked<FGameMenuBuilderModule>("GameMenuBuilder");
+		FModuleManager::LoadModuleChecked<IGameMenuBuilderModule>("GameMenuBuilder");
 	}	
 
 	TArray<FString> SplitString;
diff --git a/Engine/Source/Runtime/GameMenuBuilder/Private/SGameMenuItemWidget.cpp b/Engine/Source/Runtime/GameMenuBuilder/Private/SGameMenuItemWidget.cpp
--- a/Engine/Source/Runtime/GameMenuBuilder/Private/SGameMenuItemWidget.cpp
+++ b/Engine/Source/Runtime/GameMenuBuilder/Private/SGameMenuItemWidget.cpp
@@ -11,7 +11,6 @@
 #include "GameMenuBuilderStyle.h"
 #include "SGameMenuPageWidget.h"
 #include "GameMenuPage.h"
-#include "GameMenuBuilder.h"
 #include "Engine/World.h"
 
 void SGameMenuItemWidget::Construct(const FArguments& InArgs)
diff --git a/Engine/Source/Runtime/GameMenuBuilder/Public/GameMenuBuilder.h b/Engine/Source/Runtime/GameMenuBuilder/Public/GameMenuBuilder.h
--- a/Engine/Source/Runtime/GameMenuBuilder/Public/GameMenuBuilder.h
+++ b/Engine/Source/Runtime/GameMenuBuilder/Public/GameMenuBuilder.h
@@ -2,18 +2,15 @@
 
 #pragma once
 
-#include "CoreMinimal.h"
-#include "GameMenuBuilderModule.h"
-
+#include "Misc/MonolithicHeaderBoilerplate.h"
+MONOLITHIC_HEADER_BOILERPLATE()
 
-class FGameMenuBuilderModule : public IGameMenuBuilderModule
-{
-	virtual void StartupModule() override
-	{
-		
-	}
-	void SetStyleName()
-	{
-	}
+#include "GameMenuBuilderModule.h"
 
-};
+#include "SlateBasics.h"
+#include "GameMenuBuilderStyle.h"
+#include "GameMenuWidgetStyle.h"
+#include "SGameMenuPageWidget.h"
+#include "SGameMenuItemWidget.h"
+#include "GameMenuItem.h"
+#include "GameMenuPage.h"
diff --git a/Engine/Source/Runtime/GameMenuBuilder/Public/GameMenuBuilderModule.h b/Engine/Source/Runtime/GameMenuBuilder/Public/GameMenuBuilderModule.h
--- a/Engine/Source/Runtime/GameMenuBuilder/Public/GameMenuBuilderModule.h
+++ b/Engine/Source/Runtime/GameMenuBuilder/Public/GameMenuBuilderModule.h
@@ -27,3 +27,14 @@ public:
 	}
 };
 
+class FGameMenuBuilderModule : public IGameMenuBuilderModule
+{
+	virtual void StartupModule() override
+	{
+
+	}
+	void SetStyleName()
+	{
+	}
+
+};
diff --git a/Engine/Source/Runtime/MessagingRpc/Public/IMessageRpcServer.h b/Engine/Source/Runtime/MessagingRpc/Public/IMessageRpcServer.h
--- a/Engine/Source/Runtime/MessagingRpc/Public/IMessageRpcServer.h
+++ b/Engine/Source/Runtime/MessagingRpc/Public/IMessageRpcServer.h
@@ -3,6 +3,7 @@
 
 #include "CoreMinimal.h"
 #include "IMessageContext.h"
+#include "Async/IAsyncTask.h"
 #include "Async/AsyncResult.h"
 #include "IMessageRpcHandler.h"
 #include "IMessageRpcReturn.h"
diff --git a/Engine/Source/Runtime/WebBrowser/Private/IOS/IOSPlatformWebBrowser.cpp b/Engine/Source/Runtime/WebBrowser/Private/IOS/IOSPlatformWebBrowser.cpp
--- a/Engine/Source/Runtime/WebBrowser/Private/IOS/IOSPlatformWebBrowser.cpp
+++ b/Engine/Source/Runtime/WebBrowser/Private/IOS/IOSPlatformWebBrowser.cpp
@@ -1,11 +1,11 @@
 // Copyright 1998-2016 Epic Games, Inc. All Rights Reserved.
 
-#include "CoreMinimal.h"
+#include "IOSPlatformWebBrowser.h"
 
 #if PLATFORM_IOS
-#include "IOSPlatformWebBrowser.h"
 #include "IOSView.h"
 #include "IOSAppDelegate.h"
+#include "Widgets/SLeafWidget.h"
 
 #import <UIKit/UIKit.h>
 
diff --git a/Engine/Source/Runtime/WebBrowser/Private/IOS/IOSPlatformWebBrowser.h b/Engine/Source/Runtime/WebBrowser/Private/IOS/IOSPlatformWebBrowser.h
--- a/Engine/Source/Runtime/WebBrowser/Private/IOS/IOSPlatformWebBrowser.h
+++ b/Engine/Source/Runtime/WebBrowser/Private/IOS/IOSPlatformWebBrowser.h
@@ -2,10 +2,12 @@
 
 #pragma once
 
+#include "CoreMinimal.h"
+
 #if PLATFORM_IOS
 
 #include "IWebBrowserWindow.h"
-#include "SlateCore.h"
+#include "Widgets/SWindow.h"
 #import <UIKit/UIKit.h>
 #import <UIKit/UIWebView.h>
 
