Fixing issue with FTest Screenshots where having a project on a separate drive from the editor executable will cause screenshots to fail to save
FTest automation screenshot output dirs were configured as relative to FPaths::RootDir() instead of FPaths::ProjectDir(), so various path-name errors occurred when the ProjectDir was not located on the same drive as RootDir

diff --git a/Engine/Source/Runtime/AutomationWorker/Private/AutomationWorkerModule.cpp b/Engine/Source/Runtime/AutomationWorker/Private/AutomationWorkerModule.cpp
--- a/Engine/Source/Runtime/AutomationWorker/Private/AutomationWorkerModule.cpp
+++ b/Engine/Source/Runtime/AutomationWorker/Private/AutomationWorkerModule.cpp
@@ -425,7 +425,7 @@ void FAutomationWorkerModule::HandleScreenShotCapturedWithName(const TArray<FCol
 	{
 		FAutomationWorkerScreenImage* Message = new FAutomationWorkerScreenImage();
 
-		Message->ScreenShotName = FPaths::RootDir() / Data.Path;
+		Message->ScreenShotName = FPaths::ProjectDir() / Data.Path;
 		FPaths::MakePathRelativeTo(Message->ScreenShotName, *FPaths::AutomationDir());
 		Message->ScreenImage = CompressedBitmap;
 		Message->Metadata = Metadata;
diff --git a/Engine/Source/Runtime/Engine/Private/Tests/AutomationCommon.cpp b/Engine/Source/Runtime/Engine/Private/Tests/AutomationCommon.cpp
--- a/Engine/Source/Runtime/Engine/Private/Tests/AutomationCommon.cpp
+++ b/Engine/Source/Runtime/Engine/Private/Tests/AutomationCommon.cpp
@@ -87,7 +87,7 @@ namespace AutomationCommon
 		FString PathName = FPaths::AutomationDir() + TestName / FPlatformProperties::IniPlatformName();
 		PathName = PathName + TEXT("/") + GetRenderDetailsString();
 
-		FPaths::MakePathRelativeTo(PathName, *FPaths::RootDir());
+		FPaths::MakePathRelativeTo(PathName, *FPaths::ProjectDir());
 
 		OutScreenshotName = FString::Printf(TEXT("%s/%s.png"), *PathName, *FPlatformMisc::GetDeviceId());
 	}
