Fix Automation Worker Instance Name not being unique for some console platforms.
+ Add support for device tagging through Automation Worker

diff --git a/Engine/Source/Programs/AutomationTool/Gauntlet/Unreal/Base/Gauntlet.UnrealSession.cs b/Engine/Source/Programs/AutomationTool/Gauntlet/Unreal/Base/Gauntlet.UnrealSession.cs
--- a/Engine/Source/Programs/AutomationTool/Gauntlet/Unreal/Base/Gauntlet.UnrealSession.cs
+++ b/Engine/Source/Programs/AutomationTool/Gauntlet/Unreal/Base/Gauntlet.UnrealSession.cs
@@ -732,6 +732,12 @@ namespace Gauntlet
 					// todo - should this be elsewhere?
 					AppConfig.Sandbox = Sandbox;
 
+					// Tag the device for report result
+					if(BuildHostPlatform.Current.Platform != Device.Platform)
+					{
+						AppConfig.CommandLineParams.Add("DeviceTag", Device.Name);
+					}
+
 					IAppInstall Install = null;
 
 					IDeviceUsageReporter.RecordStart(Device.Name, Device.Platform, IDeviceUsageReporter.EventType.Device, IDeviceUsageReporter.EventState.Success);
diff --git a/Engine/Source/Runtime/AutomationWorker/Private/AutomationWorkerModule.cpp b/Engine/Source/Runtime/AutomationWorker/Private/AutomationWorkerModule.cpp
--- a/Engine/Source/Runtime/AutomationWorker/Private/AutomationWorkerModule.cpp
+++ b/Engine/Source/Runtime/AutomationWorker/Private/AutomationWorkerModule.cpp
@@ -160,6 +160,8 @@ void FAutomationWorkerModule::Initialize()
 	ExecutionCount = INDEX_NONE;
 	bExecutingNetworkCommandResults = false;
 	bSendAnalytics = false;
+
+	FParse::Value(FCommandLine::Get(), TEXT("-DeviceTag="), DeviceTag);
 }
 
 void FAutomationWorkerModule::ReportNetworkCommandComplete()
@@ -329,8 +331,11 @@ void FAutomationWorkerModule::SendWorkerFound(const FMessageAddress& ControllerA
 	FString OSVersionString = OSMajorVersionString + TEXT(" ") + OSSubVersionString;
 	FString CPUModelString = FPlatformMisc::GetCPUBrand().TrimStart();
 
-	Response->DeviceName = FPlatformProcess::ComputerName();
-	Response->InstanceName = FString::Printf(TEXT("%s-%i"), FPlatformProcess::ComputerName(), FPlatformProcess::GetCurrentProcessId());
+	FString DeviceName = DeviceTag.IsEmpty() ? FPlatformProcess::ComputerName() : DeviceTag;
+	FString DeviceId = FPlatformMisc::GetDeviceId().IsEmpty() ? DeviceName : FPlatformMisc::GetDeviceId();
+
+	Response->DeviceName = DeviceName;
+	Response->InstanceName = FString::Printf(TEXT("%s-%s"), *DeviceId, *FApp::GetSessionId().ToString());
 	Response->Platform = FPlatformProperties::PlatformName();
 	Response->SessionId = FApp::GetSessionId();
 	Response->OSVersionName = OSVersionString;
diff --git a/Engine/Source/Runtime/AutomationWorker/Private/AutomationWorkerModule.h b/Engine/Source/Runtime/AutomationWorker/Private/AutomationWorkerModule.h
--- a/Engine/Source/Runtime/AutomationWorker/Private/AutomationWorkerModule.h
+++ b/Engine/Source/Runtime/AutomationWorker/Private/AutomationWorkerModule.h
@@ -162,4 +162,7 @@ private:
 
 	/** Delegate to fire when the test is complete */
 	FStopTestEvent StopTestEvent;
+
+	/** Tag for the device on which the worker is running */
+	FString DeviceTag;
 };
