Put back LLM scope and float to double change in test Automation modules

#rnx

diff --git a/Engine/Source/Developer/AutomationController/Private/AutomationControllerManager.h b/Engine/Source/Developer/AutomationController/Private/AutomationControllerManager.h
--- a/Engine/Source/Developer/AutomationController/Private/AutomationControllerManager.h
+++ b/Engine/Source/Developer/AutomationController/Private/AutomationControllerManager.h
@@ -574,7 +574,7 @@ private:
 		/** The test runners message address */
 		FMessageAddress OwnerMessageAddress;
 		/** The time since we had a ping from the instance*/
-		float LastPingTime;
+		double LastPingTime;
 		/** Is Confirmed to be running */
 		bool bConfirmed;
 	};
diff --git a/Engine/Source/Runtime/AutomationWorker/Private/AutomationWorkerModule.cpp b/Engine/Source/Runtime/AutomationWorker/Private/AutomationWorkerModule.cpp
--- a/Engine/Source/Runtime/AutomationWorker/Private/AutomationWorkerModule.cpp
+++ b/Engine/Source/Runtime/AutomationWorker/Private/AutomationWorkerModule.cpp
@@ -489,6 +489,8 @@ void FAutomationWorkerModule::HandleScreenShotAndTraceCapturedWithName(const TAr
 
 void FAutomationWorkerModule::HandleRunTestsMessage( const FAutomationWorkerRunTests& Message, const TSharedRef<IMessageContext, ESPMode::ThreadSafe>& Context )
 {
+	LLM_SCOPE_BYNAME(TEXT("AutomationTest/Worker"));
+	
 	if (TestRequesterAddress.IsValid() && !TestName.IsEmpty())
 	{
 		// Worker is already running a test.
