Simplify RHI string when excluding tests by removing the shader model extension
- follow up to 21852012

#rnx

diff --git a/Engine/Source/Runtime/AutomationTest/Public/AutomationTestExcludelist.h b/Engine/Source/Runtime/AutomationTest/Public/AutomationTestExcludelist.h
--- a/Engine/Source/Runtime/AutomationTest/Public/AutomationTestExcludelist.h
+++ b/Engine/Source/Runtime/AutomationTest/Public/AutomationTestExcludelist.h
@@ -8,10 +8,10 @@
 UENUM(Meta = (Bitflags, UseEnumValuesAsMaskValuesInEditor = "true"))
 enum class ERHI_Flags : uint8
 {
-	DirectX11	= 0x01 UMETA(DisplayName = "DirectX 11"),
-	DirectX12	= 0x02 UMETA(DisplayName = "DirectX 12"),
-	Vulkan		= 0x04 UMETA(DisplayName = "Vulkan"),
-	Metal		= 0x08 UMETA(DisplayName = "Metal"),
+	DirectX11	= 1 << 0 UMETA(DisplayName = "DirectX 11"),
+	DirectX12	= 1 << 1 UMETA(DisplayName = "DirectX 12"),
+	Vulkan		= 1 << 2 UMETA(DisplayName = "Vulkan"),
+	Metal		= 1 << 3 UMETA(DisplayName = "Metal"),
 	NUM			UMETA(Hidden)
 };
 
diff --git a/Engine/Source/Runtime/AutomationWorker/Private/AutomationWorkerModule.cpp b/Engine/Source/Runtime/AutomationWorker/Private/AutomationWorkerModule.cpp
--- a/Engine/Source/Runtime/AutomationWorker/Private/AutomationWorkerModule.cpp
+++ b/Engine/Source/Runtime/AutomationWorker/Private/AutomationWorkerModule.cpp
@@ -486,6 +486,17 @@ void FAutomationWorkerModule::HandleScreenShotAndTraceCapturedWithName(const TAr
 }
 #endif
 
+FString GetRHIForAutomation()
+{
+	// Remove any extra information in () from RHI string
+	FString RHI = FApp::GetGraphicsRHI();
+	int Pos;
+	if (RHI.FindChar(*TEXT("("), Pos))
+	{
+		RHI = RHI.Left(Pos).TrimEnd();
+	}
+	return RHI;
+}
 
 void FAutomationWorkerModule::HandleRunTestsMessage( const FAutomationWorkerRunTests& Message, const TSharedRef<IMessageContext, ESPMode::ThreadSafe>& Context )
 {
@@ -521,7 +532,8 @@ void FAutomationWorkerModule::HandleRunTestsMessage( const FAutomationWorkerRunT
 	FName SkipReason;
 	bool bWarn(false);
 	UAutomationTestExcludelist* Excludelist = UAutomationTestExcludelist::Get();
-	if (Excludelist->IsTestExcluded(Message.FullTestPath, FApp::GetGraphicsRHI(), &SkipReason, &bWarn))
+	static FString RHI = GetRHIForAutomation();
+	if (Excludelist->IsTestExcluded(Message.FullTestPath, RHI, &SkipReason, &bWarn))
 	{
 		FString SkippingMessage = FString::Format(TEXT("Test Skipped. Name={{0}} Reason={{1}} Path={{2}}"),
 			{ *Message.BeautifiedTestName, *SkipReason.ToString(), *Message.FullTestPath });
