PR #10022: Include instance ID in automation InstanceName (Contributed by [REDACTED_PERSON])
Adds the value of FApp::GetInstanceId() to the instance field in index.json such that it is now (hostname)-(session ID)-(instance ID)

diff --git a/Engine/Source/Runtime/AutomationWorker/Private/AutomationWorkerModule.cpp b/Engine/Source/Runtime/AutomationWorker/Private/AutomationWorkerModule.cpp
--- a/Engine/Source/Runtime/AutomationWorker/Private/AutomationWorkerModule.cpp
+++ b/Engine/Source/Runtime/AutomationWorker/Private/AutomationWorkerModule.cpp
@@ -300,7 +300,7 @@ void FAutomationWorkerModule::SendWorkerFound(const FMessageAddress& ControllerA
 	FString DeviceId = FPlatformMisc::GetDeviceId().IsEmpty() ? DeviceName : FPlatformMisc::GetDeviceId();
 
 	Response->DeviceName = DeviceName;
-	Response->InstanceName = FString::Printf(TEXT("%s-%s"), *DeviceId, *FApp::GetSessionId().ToString());
+	Response->InstanceName = FString::Printf(TEXT("%s-%s-%s"), *DeviceId, *FApp::GetSessionId().ToString(), *FApp::GetInstanceId().ToString());
 	Response->Platform = FPlatformProperties::PlatformName();
 	Response->SessionId = FApp::GetSessionId();
 	Response->OSVersionName = OSVersionString;
