Disable testing in Client net mode

diff --git a/Engine/Source/Runtime/AutomationWorker/AutomationWorker.Build.cs b/Engine/Source/Runtime/AutomationWorker/AutomationWorker.Build.cs
--- a/Engine/Source/Runtime/AutomationWorker/AutomationWorker.Build.cs
+++ b/Engine/Source/Runtime/AutomationWorker/AutomationWorker.Build.cs
@@ -38,6 +38,11 @@ namespace UnrealBuildTool.Rules
 				PrivateDependencyModuleNames.Add("Engine");
 				PrivateDependencyModuleNames.Add("RHI");
 			}
+
+			if (Target.bCompileAgainstEditor)
+			{
+				PrivateDependencyModuleNames.Add("UnrealEd");
+			}
 		}
 	}
 }
diff --git a/Engine/Source/Runtime/AutomationWorker/Private/AutomationWorkerModule.cpp b/Engine/Source/Runtime/AutomationWorker/Private/AutomationWorkerModule.cpp
--- a/Engine/Source/Runtime/AutomationWorker/Private/AutomationWorkerModule.cpp
+++ b/Engine/Source/Runtime/AutomationWorker/Private/AutomationWorkerModule.cpp
@@ -15,7 +15,6 @@
 #include "Modules/ModuleManager.h"
 
 #if WITH_ENGINE
-	#include "Engine/Engine.h"
 	#include "Engine/GameViewportClient.h"
 	#include "ImageUtils.h"
 	#include "Tests/AutomationCommon.h"
@@ -23,6 +22,7 @@
 #endif
 
 #if WITH_EDITOR
+	#include "Settings/LevelEditorPlaySettings.h"
 	#include "AssetRegistry/AssetRegistryModule.h"
 #endif
 
@@ -33,7 +33,6 @@ DEFINE_LOG_CATEGORY_STATIC(LogAutomationWorker, Log, All);
 
 IMPLEMENT_MODULE(FAutomationWorkerModule, AutomationWorker);
 
-
 /* IModuleInterface interface
  *****************************************************************************/
 
@@ -204,7 +203,7 @@ void FAutomationWorkerModule::ReportTestComplete()
 			Message->ExecutionCount = ExecutionCount;
 			Message->State = bSuccess ? EAutomationState::Success : EAutomationState::Fail;
 			Message->Duration = ExecutionInfo.Duration;
-			Message->Entries = ExecutionInfo.GetEntries();
+			Message->Entries = ExecutionInfo.GetEntries(); 
 			Message->WarningTotal = ExecutionInfo.GetWarningTotal();
 			Message->ErrorTotal = ExecutionInfo.GetErrorTotal();
 
@@ -498,41 +497,82 @@ FString GetRHIForAutomation()
 	return RHI;
 }
 
-void FAutomationWorkerModule::HandleRunTestsMessage( const FAutomationWorkerRunTests& Message, const TSharedRef<IMessageContext, ESPMode::ThreadSafe>& Context )
+bool FAutomationWorkerModule::AreTestsShouldBeSkipped( const FAutomationWorkerRunTests& Message, const TSharedRef<IMessageContext, ESPMode::ThreadSafe>& Context, FAutomationWorkerRunTestsReply*& OutReply )
 {
-	UE_LOG(LogAutomationWorker, Log, TEXT("Received RunTests %s from %s"), *Message.BeautifiedTestName, *Context->GetSender().ToString());
+	// Note that short circuit evaluation is used here to be able to form correct output values. The order of check functions call is important
+	return
+		AreTestsShouldBeSkippedFormalCheck(Message, Context, OutReply) ||
+		AreTestsShouldBeSkippedNetworkModeCheck(Message, OutReply) ||
+		AreTestsShouldBeSkippedExcludeListCheck(Message, OutReply);
+}
 
-	LLM_SCOPE_BYNAME(TEXT("AutomationTest/Worker"));
-	
+bool FAutomationWorkerModule::AreTestsShouldBeSkippedFormalCheck( const FAutomationWorkerRunTests& Message, const TSharedRef<IMessageContext, ESPMode::ThreadSafe>& Context, FAutomationWorkerRunTestsReply*& OutReply )
+{
 	if (TestRequesterAddress.IsValid() && !TestName.IsEmpty())
 	{
 		if (TestRequesterAddress == Context->GetSender())
 		{
 			UE_LOG(LogAutomationWorker, Log, TEXT("Worker is already running test '%s' from %s. Request is ignored."), *BeautifiedTestName, *TestRequesterAddress.ToString());
-			return;
+			return true;
 		}
 
-		FString LogMessage = FString::Format(TEXT("Worker is already running test '%s' from %s. '%s' won't be run."), 
+		FString LogMessage = FString::Format(TEXT("Worker is already running test '%s' from %s. '%s' won't be run."),
 			{ *BeautifiedTestName, *TestRequesterAddress.ToString(), *Message.BeautifiedTestName });
 		UE_LOG(LogAutomationWorker, Warning, TEXT("%s"), *LogMessage);
-
+		
 		// Let the sender know it won't happen
-		FAutomationWorkerRunTestsReply* OutMessage = FMessageEndpoint::MakeMessage<FAutomationWorkerRunTestsReply>();
-		OutMessage->TestName = Message.TestName;
-		OutMessage->ExecutionCount = Message.ExecutionCount;
-		OutMessage->State = EAutomationState::Skipped;
-		OutMessage->Entries.Add(FAutomationExecutionEntry(FAutomationEvent(EAutomationEventType::Error, LogMessage)));
-		OutMessage->ErrorTotal = 1;
-		MessageEndpoint->Send(OutMessage, Context->GetSender());
+		OutReply = FMessageEndpoint::MakeMessage<FAutomationWorkerRunTestsReply>();
+		OutReply->TestName = Message.TestName;
+		OutReply->ExecutionCount = Message.ExecutionCount;
+		OutReply->State = EAutomationState::Skipped;
+		OutReply->Entries.Add(FAutomationExecutionEntry(FAutomationEvent(EAutomationEventType::Error, LogMessage)));
+		OutReply->ErrorTotal = 1;
+
+		return true;
+	}
 
-		return;
+	return false;
+}
+
+bool FAutomationWorkerModule::AreTestsShouldBeSkippedNetworkModeCheck( const FAutomationWorkerRunTests& Message, FAutomationWorkerRunTestsReply*& OutReply )
+{
+#if WITH_EDITOR
+	// Do we need to skip the test due to unsupported net mode (EPlayNetMode::PIE_Client is currently unsupported)
+	const ULevelEditorPlaySettings* const PlayInEditorSettings = GetDefault<ULevelEditorPlaySettings>();
+	if (PlayInEditorSettings)
+	{
+		EPlayNetMode NetMode;
+		const bool bGetPlayNetModeSuccess(PlayInEditorSettings->GetPlayNetMode(NetMode));
+
+		if (bGetPlayNetModeSuccess && NetMode == EPlayNetMode::PIE_Client)
+		{
+			FString SkipReason = TEXT("Testing in current net mode is not supported");
+			FString LogMessage = FString::Format(TEXT("Test Skipped. Name={{0}} Reason={{1}} Path={{2}}"),
+				{ *Message.BeautifiedTestName, *SkipReason, *Message.FullTestPath });
+			UE_LOG(LogAutomationWorker, Warning, TEXT("%s"), *LogMessage)
+
+			OutReply = FMessageEndpoint::MakeMessage<FAutomationWorkerRunTestsReply>();
+			OutReply->TestName = Message.TestName;
+			OutReply->ExecutionCount = Message.ExecutionCount;
+			OutReply->State = EAutomationState::Skipped;
+			OutReply->Entries.Add(FAutomationExecutionEntry(FAutomationEvent(
+				EAutomationEventType::Info, FString::Printf(TEXT("Skipping test. Reason: %s"), *SkipReason))));
+
+			return true;
+		}
 	}
+#endif // WITH_EDITOR
+
+	return false;
+}
 
-	// Do we need to skip the test
+bool FAutomationWorkerModule::AreTestsShouldBeSkippedExcludeListCheck( const FAutomationWorkerRunTests& Message, FAutomationWorkerRunTestsReply*& OutReply )
+{
 	FName SkipReason;
 	bool bWarn(false);
 	UAutomationTestExcludelist* Excludelist = UAutomationTestExcludelist::Get();
 	static FString RHI = GetRHIForAutomation();
+
 	if (Excludelist->IsTestExcluded(Message.FullTestPath, RHI, &SkipReason, &bWarn))
 	{
 		FString SkippingMessage = FString::Format(TEXT("Test Skipped. Name={{0}} Reason={{1}} Path={{2}}"),
@@ -546,13 +586,33 @@ void FAutomationWorkerModule::HandleRunTestsMessage( const FAutomationWorkerRunT
 			UE_LOG(LogAutomationWorker, Display, TEXT("%s"), *SkippingMessage);
 		}
 
-		FAutomationWorkerRunTestsReply* OutMessage = FMessageEndpoint::MakeMessage<FAutomationWorkerRunTestsReply>();
-		OutMessage->TestName = Message.TestName;
-		OutMessage->ExecutionCount = Message.ExecutionCount;
-		OutMessage->State = EAutomationState::Skipped;
-		OutMessage->Entries.Add(FAutomationExecutionEntry(FAutomationEvent(EAutomationEventType::Info, FString::Printf(TEXT("Skipping test because of exclude list: %s"), *SkipReason.ToString()))));
-		MessageEndpoint->Send(OutMessage, Context->GetSender());
+		OutReply = FMessageEndpoint::MakeMessage<FAutomationWorkerRunTestsReply>();
+		OutReply->TestName = Message.TestName;
+		OutReply->ExecutionCount = Message.ExecutionCount;
+		OutReply->State = EAutomationState::Skipped;
+		OutReply->Entries.Add(FAutomationExecutionEntry(FAutomationEvent(EAutomationEventType::Info, FString::Printf(TEXT("Skipping test because of exclude list: %s"), *SkipReason.ToString()))));
+
+		return true;
+	}
+
+	return false;
+}
+
+void FAutomationWorkerModule::HandleRunTestsMessage( const FAutomationWorkerRunTests& Message, const TSharedRef<IMessageContext, ESPMode::ThreadSafe>& Context )
+{
+	UE_LOG(LogAutomationWorker, Log, TEXT("Received RunTests %s from %s"), *Message.BeautifiedTestName, *Context->GetSender().ToString());
+
+	LLM_SCOPE_BYNAME(TEXT("AutomationTest/Worker"));
+	
+	FAutomationWorkerRunTestsReply* OutReply = nullptr;
+	if (AreTestsShouldBeSkipped(Message, Context, OutReply))
+	{
+		if (nullptr != OutReply)
+		{
+			MessageEndpoint->Send(OutReply, Context->GetSender());
+		}
 
+		// Skipping the tests
 		return;
 	}
 
diff --git a/Engine/Source/Runtime/AutomationWorker/Private/AutomationWorkerModule.h b/Engine/Source/Runtime/AutomationWorker/Private/AutomationWorkerModule.h
--- a/Engine/Source/Runtime/AutomationWorker/Private/AutomationWorkerModule.h
+++ b/Engine/Source/Runtime/AutomationWorker/Private/AutomationWorkerModule.h
@@ -129,6 +129,18 @@ private:
 	/** Helper for Performance Capture Analytics. */
 	void RecordPerformanceAnalytics(const FAutomationPerformanceSnapshot& PerfSnapshot);
 
+	/** The main check if tests corresponding to the run tests message should be skipped. */
+	bool AreTestsShouldBeSkipped(const FAutomationWorkerRunTests& Message, const TSharedRef<IMessageContext, ESPMode::ThreadSafe>& Context, struct FAutomationWorkerRunTestsReply*& OutReply);
+
+	/** Check if the run tests request should be skipped due to formal reasons. */
+	bool AreTestsShouldBeSkippedFormalCheck(const FAutomationWorkerRunTests& Message, const TSharedRef<IMessageContext, ESPMode::ThreadSafe>& Context	, struct FAutomationWorkerRunTestsReply*& OutReply);
+
+	/** Check if the run tests request should be skipped due to unsopported network mode. */
+	bool AreTestsShouldBeSkippedNetworkModeCheck(const FAutomationWorkerRunTests& Message, struct FAutomationWorkerRunTestsReply*& OutReply);
+
+	/** Check if the run tests request should be skipped due to exclude list. */
+	bool AreTestsShouldBeSkippedExcludeListCheck(const FAutomationWorkerRunTests& Message, struct FAutomationWorkerRunTestsReply*& OutReply);
+
 private:
 
 	/** The collection of test data we are to send to a controller. */
