Add FAutomationTestBase::CanRunInEnvironment methods with basic implementation

diff --git a/Engine/Source/Runtime/AutomationWorker/Private/AutomationWorkerModule.cpp b/Engine/Source/Runtime/AutomationWorker/Private/AutomationWorkerModule.cpp
--- a/Engine/Source/Runtime/AutomationWorker/Private/AutomationWorkerModule.cpp
+++ b/Engine/Source/Runtime/AutomationWorker/Private/AutomationWorkerModule.cpp
@@ -498,6 +498,25 @@ FString GetRHIForAutomation()
 	return RHI;
 }
 
+bool FAutomationWorkerModule::IsTestExcluded(const FString& InTestToRun, FString* OutReason, bool* OutWarn) const
+{
+	FName SkipReason;
+	UAutomationTestExcludelist* Excludelist = UAutomationTestExcludelist::Get();
+	static FString RHI = GetRHIForAutomation();
+	if (Excludelist->IsTestExcluded(InTestToRun, RHI, &SkipReason, OutWarn))
+	{
+		if (OutReason)
+		{
+			(*OutReason) = (SkipReason.IsNone() ? TEXT("unknown reason") : SkipReason.ToString());
+			(*OutReason) += TEXT(" [config]");
+		}
+
+		return true;
+	}
+
+	return false;
+}
+
 void FAutomationWorkerModule::HandleRunTestsMessage( const FAutomationWorkerRunTests& Message, const TSharedRef<IMessageContext, ESPMode::ThreadSafe>& Context )
 {
 	UE_LOG(LogAutomationWorker, Log, TEXT("Received RunTests %s from %s"), *Message.BeautifiedTestName, *Context->GetSender().ToString());
@@ -529,14 +548,16 @@ void FAutomationWorkerModule::HandleRunTestsMessage( const FAutomationWorkerRunT
 	}
 
 	// Do we need to skip the test
-	FName SkipReason;
+	FString SkipReason;
 	bool bWarn(false);
+	FAutomationTestFramework& AutomationTestFramework = FAutomationTestFramework::Get();
 	UAutomationTestExcludelist* Excludelist = UAutomationTestExcludelist::Get();
-	static FString RHI = GetRHIForAutomation();
-	if (Excludelist->IsTestExcluded(Message.FullTestPath, RHI, &SkipReason, &bWarn))
+	if (!AutomationTestFramework.CanRunTestInEnvironment(Message.TestName, &SkipReason, &bWarn)
+		|| IsTestExcluded(Message.FullTestPath, &SkipReason, &bWarn))
 	{
 		FString SkippingMessage = FString::Format(TEXT("Test Skipped. Name={{0}} Reason={{1}} Path={{2}}"),
-			{ *Message.BeautifiedTestName, *SkipReason.ToString(), *Message.FullTestPath });
+			{ *Message.BeautifiedTestName, *SkipReason, *Message.FullTestPath });
+
 		if (bWarn)
 		{
 			UE_LOG(LogAutomationWorker, Warning, TEXT("%s"), *SkippingMessage);
@@ -550,7 +571,7 @@ void FAutomationWorkerModule::HandleRunTestsMessage( const FAutomationWorkerRunT
 		OutMessage->TestName = Message.TestName;
 		OutMessage->ExecutionCount = Message.ExecutionCount;
 		OutMessage->State = EAutomationState::Skipped;
-		OutMessage->Entries.Add(FAutomationExecutionEntry(FAutomationEvent(EAutomationEventType::Info, FString::Printf(TEXT("Skipping test because of exclude list: %s"), *SkipReason.ToString()))));
+		OutMessage->Entries.Add(FAutomationExecutionEntry(FAutomationEvent(EAutomationEventType::Info, FString::Printf(TEXT("Skipping test: %s"), *SkipReason))));
 		MessageEndpoint->Send(OutMessage, Context->GetSender());
 
 		return;
diff --git a/Engine/Source/Runtime/AutomationWorker/Private/AutomationWorkerModule.h b/Engine/Source/Runtime/AutomationWorker/Private/AutomationWorkerModule.h
--- a/Engine/Source/Runtime/AutomationWorker/Private/AutomationWorkerModule.h
+++ b/Engine/Source/Runtime/AutomationWorker/Private/AutomationWorkerModule.h
@@ -129,6 +129,9 @@ private:
 	/** Helper for Performance Capture Analytics. */
 	void RecordPerformanceAnalytics(const FAutomationPerformanceSnapshot& PerfSnapshot);
 
+	/** Checks whether the test is in exclusion list. */
+	bool IsTestExcluded(const FString& InTestToRun, FString* OutReason, bool* OutWarn) const;
+
 private:
 
 	/** The collection of test data we are to send to a controller. */
diff --git a/Engine/Source/Runtime/Core/Private/Misc/AutomationTest.cpp b/Engine/Source/Runtime/Core/Private/Misc/AutomationTest.cpp
--- a/Engine/Source/Runtime/Core/Private/Misc/AutomationTest.cpp
+++ b/Engine/Source/Runtime/Core/Private/Misc/AutomationTest.cpp
@@ -918,6 +918,45 @@ bool FAutomationTestFramework::InternalStopTest(FAutomationTestExecutionInfo& Ou
 	return bTestSuccessful;
 }
 
+bool FAutomationTestFramework::CanRunTestInEnvironment(const FString& InTestToRun, FString* OutReason, bool* OutWarn) const
+{
+	FString TestClassName;
+	FString TestParameters;
+	if (!InTestToRun.Split(TEXT(" "), &TestClassName, &TestParameters, ESearchCase::CaseSensitive))
+	{
+		TestClassName = InTestToRun;
+	}
+
+	if (!ContainsTest(TestClassName))
+	{
+		return false;
+	}
+
+	const FAutomationTestBase* const Test = *(AutomationTestClassNameToInstanceMap.Find(TestClassName));
+
+	if (nullptr == Test)
+	{
+		return false;
+	}
+
+	if (!Test->CanRunInEnvironment(TestParameters, OutReason, OutWarn))
+	{
+		if (nullptr != OutReason)
+		{
+			if (OutReason->IsEmpty())
+			{
+				*OutReason = TEXT("unknown reason");
+			}
+
+			*OutReason += TEXT(" [code]");
+		}
+		
+		return false;
+	}
+
+	return true;
+}
+
 void FAutomationTestFramework::AddAnalyticsItemToCurrentTest( const FString& AnalyticsItem )
 {
 	if( CurrentTest != nullptr )
diff --git a/Engine/Source/Runtime/Core/Public/Misc/AutomationTest.h b/Engine/Source/Runtime/Core/Public/Misc/AutomationTest.h
--- a/Engine/Source/Runtime/Core/Public/Misc/AutomationTest.h
+++ b/Engine/Source/Runtime/Core/Public/Misc/AutomationTest.h
@@ -1047,6 +1047,17 @@ public:
 
 	void NotifyScreenshotTakenAndCompared();
 
+	/**
+	 * Internal helper method designed to check if the given test is able to run in the current environment.
+	 *
+	 * @param	InTestToRun test name
+	 * @param	OutReason the related reason of the skipping
+	 * @param	OutWarn the related warning of the skipping
+	 *
+	 * @return	true if the test is able to run; false if it is unable to run.
+	 */
+	bool CanRunTestInEnvironment(const FString& InTestToRun, FString* OutReason, bool* OutWarn) const;
+
 private:
 
 	/** Special output device used during automation testing to gather messages that happen during tests */
@@ -1256,7 +1267,6 @@ private:
 	bool bCaptureStack;
 };
 
-
 /** Simple abstract base class for all automation tests */
 class CORE_API FAutomationTestBase
 {
@@ -1615,6 +1625,14 @@ public:
 		ExecutionInfo.PopContext();
 	}
 
+	/** Checks if the test is able to run in the current environment. */
+	virtual bool CanRunInEnvironment(const FString& TestParams, FString* OutReason, bool* OutWarn) const
+	{
+		// By default the test is able to run in the current environment
+		// It is responsibility of a child class to decide if the flow should skip the corresponding test.
+		return true;
+	}
+
 public:
 
 	bool TestEqual(const TCHAR* What, const int32 Actual, const int32 Expected);
