Test Exclusion list - Differentiate null rhi from lack of information
+ Always extract Feature Level string

#rnx

diff --git a/Engine/Source/Runtime/AutomationWorker/Private/AutomationWorkerModule.cpp b/Engine/Source/Runtime/AutomationWorker/Private/AutomationWorkerModule.cpp
--- a/Engine/Source/Runtime/AutomationWorker/Private/AutomationWorkerModule.cpp
+++ b/Engine/Source/Runtime/AutomationWorker/Private/AutomationWorkerModule.cpp
@@ -491,25 +491,26 @@ void FAutomationWorkerModule::HandleScreenShotAndTraceCapturedWithName(const TAr
 TSet<FName> GetRHIForAutomation()
 {
 	FString RHI = FApp::GetGraphicsRHI();
-	if (RHI.IsEmpty())
-	{
-		RHI = LexToString(ETEST_RHI_Options::Null);
-		return TSet<FName> {FName(RHI), FName(RHI)};
-	}
-
-	// Remove any extra information in () from RHI string
-	int Pos;
-	if (RHI.FindChar(*TEXT("("), Pos))
-	{
-		RHI = RHI.Left(Pos).TrimEnd();
-	}
-
 #if WITH_ENGINE
 	FString FeatureLevel = LexToString(GMaxRHIFeatureLevel);
 #else
 	FString FeatureLevel = TEXT("N/A");
 #endif
 
+	if (RHI.IsEmpty())
+	{
+		RHI = FParse::Param(FCommandLine::Get(), TEXT("nullrhi"))? LexToString(ETEST_RHI_Options::Null) : TEXT("N/A");
+	}
+	else
+	{
+		// Remove any extra information in () from RHI string
+		int Pos;
+		if (RHI.FindChar(*TEXT("("), Pos))
+		{
+			RHI = RHI.Left(Pos).TrimEnd();
+		}
+	}
+
 	return TSet<FName> {FName(RHI), FName(FeatureLevel)};
 }
 
