Introduce OnEnteringTestSectionEvent, OnLeavingTestSectionEvent
Called when entering and leaving test section.
ie:
For a test named MyTest.Log.Validation.TestTrue
These will be called through OnEnteringTestSectionEvent before any test of the corresponding section is executed:
- MyTest
- MyTest.Log
- MyTest.Log.Validation
And the corresponding one when leaving the section after execution.
One can register a callback this way:
 FAutomationTestFramework::Get().GetOnEnteringTestSection(TEXT("MyTest.Log")).AddRaw(this, &FMyModule::MySectionCallback);
As first argument the name of the section is passed.

diff --git a/Engine/Plugins/Tests/TestSamples/Source/TestSamples/Private/TestSamples.cpp b/Engine/Plugins/Tests/TestSamples/Source/TestSamples/Private/TestSamples.cpp
--- a/Engine/Plugins/Tests/TestSamples/Source/TestSamples/Private/TestSamples.cpp
+++ b/Engine/Plugins/Tests/TestSamples/Source/TestSamples/Private/TestSamples.cpp
@@ -17,6 +17,10 @@ void FTestSamplesModule::StartupModule()
 	// Hook callbacks to Test Start and End event
 	FAutomationTestFramework::Get().OnTestStartEvent.AddRaw(this, &FTestSamplesModule::OnTestStart);
 	FAutomationTestFramework::Get().OnTestEndEvent.AddRaw(this, &FTestSamplesModule::OnTestEnd);
+
+	// Hook callbacks to Test Section Enter and Leave event
+	FAutomationTestFramework::Get().GetOnEnteringTestSection(TEXT("TestFramework")).AddRaw(this, &FTestSamplesModule::OnEnteringTestSection);
+	FAutomationTestFramework::Get().GetOnLeavingTestSection(TEXT("TestFramework")).AddRaw(this, &FTestSamplesModule::OnLeavingTestSection);
 }
 
 void FTestSamplesModule::ShutdownModule()
@@ -27,6 +31,9 @@ void FTestSamplesModule::ShutdownModule()
 
 	FAutomationTestFramework::Get().OnBeforeAllTestsEvent.RemoveAll(this);
 	FAutomationTestFramework::Get().OnAfterAllTestsEvent.RemoveAll(this);
+
+	FAutomationTestFramework::Get().GetOnEnteringTestSection(TEXT("TestFramework")).RemoveAll(this);
+	FAutomationTestFramework::Get().GetOnLeavingTestSection(TEXT("TestFramework")).RemoveAll(this);
 }
 
 void FTestSamplesModule::OnTestStart(FAutomationTestBase* Test)
@@ -55,6 +62,16 @@ void FTestSamplesModule::OnAfterAllTests()
 	UE_LOG(LogTestSamples, Verbose, TEXT("Running tests completed"));
 }
 
+void FTestSamplesModule::OnEnteringTestSection(const FString& Section)
+{
+	UE_LOG(LogTestSamples, Verbose, TEXT("Entering section %s"), *Section);
+}
+
+void FTestSamplesModule::OnLeavingTestSection(const FString& Section)
+{
+	UE_LOG(LogTestSamples, Verbose, TEXT("Leaving section %s"), *Section);
+}
+
 #undef LOCTEXT_NAMESPACE
 	
 IMPLEMENT_MODULE(FTestSamplesModule, TestSamples)
diff --git a/Engine/Plugins/Tests/TestSamples/Source/TestSamples/Public/TestSamples.h b/Engine/Plugins/Tests/TestSamples/Source/TestSamples/Public/TestSamples.h
--- a/Engine/Plugins/Tests/TestSamples/Source/TestSamples/Public/TestSamples.h
+++ b/Engine/Plugins/Tests/TestSamples/Source/TestSamples/Public/TestSamples.h
@@ -20,4 +20,7 @@ private:
 
 	void OnTestStart(FAutomationTestBase* Test);
 	void OnTestEnd(FAutomationTestBase* Test);
+
+	void OnEnteringTestSection(const FString& Section);
+	void OnLeavingTestSection(const FString& Section);
 };
diff --git a/Engine/Source/Runtime/AutomationWorker/Private/AutomationWorkerModule.cpp b/Engine/Source/Runtime/AutomationWorker/Private/AutomationWorkerModule.cpp
--- a/Engine/Source/Runtime/AutomationWorker/Private/AutomationWorkerModule.cpp
+++ b/Engine/Source/Runtime/AutomationWorker/Private/AutomationWorkerModule.cpp
@@ -2,6 +2,7 @@
 
 #include "AutomationWorkerModule.h"
 
+#include "Algo/Reverse.h"
 #include "AutomationAnalytics.h"
 #include "AutomationWorkerMessages.h"
 #include "AutomationTestExcludelist.h"
@@ -345,6 +346,7 @@ void FAutomationWorkerModule::HandleStartTestSession( const FAutomationWorkerSta
 	UE_LOG(LogAutomationWorker, Log, TEXT("Received StartTestSession from %s"), *Context->GetSender().ToString());
 
 	FAutomationTestFramework::Get().ResetTests();
+	ActiveSection.Empty();
 	FAutomationTestFramework::Get().OnBeforeAllTestsEvent.Broadcast();
 }
 
@@ -352,6 +354,9 @@ void FAutomationWorkerModule::HandleStopTestSession(const FAutomationWorkerStopT
 {
 	UE_LOG(LogAutomationWorker, Log, TEXT("Received StopTestSession from %s"), *Context->GetSender().ToString());
 
+	// Unwind Active Section
+	TriggerSectionNotifications();
+
 	FAutomationTestFramework::Get().OnAfterAllTestsEvent.Broadcast();
 }
 
@@ -655,9 +660,85 @@ void FAutomationWorkerModule::HandleRunTestsMessage( const FAutomationWorkerRunT
 	// We are not executing network command sub-commands right now
 	bExecutingNetworkCommandResults = false;
 
+	// Track active section
+	TriggerSectionNotifications();
+
 	FAutomationTestFramework::Get().StartTestByName(Message.TestName, Message.RoleIndex, Message.FullTestPath);
 }
 
+void FAutomationWorkerModule::TriggerSectionNotifications()
+{
+	if (FullTestPath.IsEmpty())
+	{
+		// Unwind
+		if (!ActiveSection.IsEmpty())
+		{
+			if (FAutomationTestFramework::Get().IsAnyOnLeavingTestSectionBound())
+			{
+				FAutomationTestFramework::Get().TriggerOnLeavingTestSection(ActiveSection);
+				int32 Pos;
+				while (ActiveSection.FindLastChar('.', Pos))
+				{
+					ActiveSection.LeftInline(Pos);
+					FAutomationTestFramework::Get().TriggerOnLeavingTestSection(ActiveSection);
+				}
+			}
+			ActiveSection.Empty();
+		}
+		return;
+	}
+
+	if (!FAutomationTestFramework::Get().IsAnyOnEnteringTestSectionBound()
+		&& !FAutomationTestFramework::Get().IsAnyOnLeavingTestSectionBound())
+	{
+		return;
+	}
+
+	// Gather nesting sections
+	TArray<FString> NestingSections;
+	{
+		int32 Pos;
+		FString Section = FullTestPath;
+		while (Section.FindLastChar('.', Pos))
+		{
+			Section.LeftInline(Pos);
+			if (ActiveSection.IsEmpty() || !ActiveSection.StartsWith(Section))
+			{
+				NestingSections.Add(Section);
+			}
+			else
+			{
+				break;
+			}
+		}
+	}
+
+	if (NestingSections.Num() > 0)
+	{
+		// Notify leaving sections
+		if (!ActiveSection.IsEmpty() && FAutomationTestFramework::Get().IsAnyOnLeavingTestSectionBound())
+		{
+			const FString TargetSection = NestingSections[0];
+			int32 Pos;
+			while (!TargetSection.StartsWith(ActiveSection) && ActiveSection.FindLastChar('.', Pos))
+			{
+				FAutomationTestFramework::Get().TriggerOnLeavingTestSection(ActiveSection);
+				ActiveSection.LeftInline(Pos);
+			}
+		}
+
+		// Notify entering sections
+		ActiveSection = NestingSections[0];
+		if (FAutomationTestFramework::Get().IsAnyOnEnteringTestSectionBound())
+		{
+			Algo::Reverse(NestingSections);
+			for (const FString& Section : NestingSections)
+			{
+				FAutomationTestFramework::Get().TriggerOnEnteringTestSection(Section);
+			}
+		}
+	}
+}
 
 void FAutomationWorkerModule::HandleStopTestsMessage(const FAutomationWorkerStopTests& Message, const TSharedRef<IMessageContext, ESPMode::ThreadSafe>& Context)
 {
diff --git a/Engine/Source/Runtime/AutomationWorker/Private/AutomationWorkerModule.h b/Engine/Source/Runtime/AutomationWorker/Private/AutomationWorkerModule.h
--- a/Engine/Source/Runtime/AutomationWorker/Private/AutomationWorkerModule.h
+++ b/Engine/Source/Runtime/AutomationWorker/Private/AutomationWorkerModule.h
@@ -139,6 +139,9 @@ private:
 	/** Checks whether the test is in exclusion list. */
 	bool IsTestExcluded(const FString& InTestToRun, FString* OutReason, bool* OutWarn) const;
 
+	/** Trigger Notifications when entering and leaving section */
+	void TriggerSectionNotifications();
+
 private:
 
 	/** The collection of test data we are to send to a controller. */
@@ -178,4 +181,7 @@ private:
 
 	/** Tag for the device on which the worker is running */
 	FString DeviceTag;
+
+	/** Tracking of active section */
+	FString ActiveSection;
 };
diff --git a/Engine/Source/Runtime/Core/Private/Misc/AutomationTest.cpp b/Engine/Source/Runtime/Core/Private/Misc/AutomationTest.cpp
--- a/Engine/Source/Runtime/Core/Private/Misc/AutomationTest.cpp
+++ b/Engine/Source/Runtime/Core/Private/Misc/AutomationTest.cpp
@@ -785,6 +785,74 @@ FOnTestScreenshotAndTraceCaptured& FAutomationTestFramework::OnScreenshotAndTrac
 	return TestScreenshotAndTraceCapturedDelegate;
 }
 
+FOnTestSectionEvent& FAutomationTestFramework::GetOnEnteringTestSection(const FString& Section)
+{
+	if (!OnEnteringTestSectionEvent.Contains(Section))
+	{
+		OnEnteringTestSectionEvent.Emplace(Section);
+	}
+
+	return *OnEnteringTestSectionEvent.Find(Section);
+}
+
+void FAutomationTestFramework::TriggerOnEnteringTestSection(const FString& Section) const
+{
+	if (const FOnTestSectionEvent* Delegate = OnEnteringTestSectionEvent.Find(Section))
+	{
+		Delegate->Broadcast(Section);
+	}
+}
+
+bool FAutomationTestFramework::IsAnyOnEnteringTestSectionBound() const
+{
+	if (!OnEnteringTestSectionEvent.IsEmpty())
+	{
+		for (auto& SectionPair : OnEnteringTestSectionEvent)
+		{
+			if (SectionPair.Value.IsBound())
+			{
+				return true;
+			}
+		}
+	}
+
+	return false;
+}
+
+FOnTestSectionEvent& FAutomationTestFramework::GetOnLeavingTestSection(const FString& Section)
+{
+	if (!OnLeavingTestSectionEvent.Contains(Section))
+	{
+		OnLeavingTestSectionEvent.Emplace(Section);
+	}
+
+	return *OnLeavingTestSectionEvent.Find(Section);
+}
+
+void FAutomationTestFramework::TriggerOnLeavingTestSection(const FString& Section) const
+{
+	if (const FOnTestSectionEvent* Delegate = OnLeavingTestSectionEvent.Find(Section))
+	{
+		Delegate->Broadcast(Section);
+	}
+}
+
+bool FAutomationTestFramework::IsAnyOnLeavingTestSectionBound() const
+{
+	if (!OnLeavingTestSectionEvent.IsEmpty())
+	{
+		for (auto& SectionPair : OnLeavingTestSectionEvent)
+		{
+			if (SectionPair.Value.IsBound())
+			{
+				return true;
+			}
+		}
+	}
+
+	return false;
+}
+
 void FAutomationTestFramework::PrepForAutomationTests()
 {
 	check(!GIsAutomationTesting);
diff --git a/Engine/Source/Runtime/Core/Public/Misc/AutomationTest.h b/Engine/Source/Runtime/Core/Public/Misc/AutomationTest.h
--- a/Engine/Source/Runtime/Core/Public/Misc/AutomationTest.h
+++ b/Engine/Source/Runtime/Core/Public/Misc/AutomationTest.h
@@ -908,6 +908,8 @@ DECLARE_MULTICAST_DELEGATE_TwoParams(FOnPerformanceDataRetrieved, bool /*bSucces
 
 DECLARE_MULTICAST_DELEGATE_OneParam(FOnTestEvent, FAutomationTestBase*);
 
+DECLARE_MULTICAST_DELEGATE_OneParam(FOnTestSectionEvent, const FString& /*Section*/);
+
 /** Class representing the main framework for running automation tests */
 class FAutomationTestFramework
 {
@@ -945,6 +947,16 @@ public:
 	/** Called after all chosen tests run have finished. */
 	FSimpleMulticastDelegate OnAfterAllTestsEvent;
 
+	/** Called entering test section. */
+	CORE_API FOnTestSectionEvent& GetOnEnteringTestSection(const FString& Section);
+	CORE_API void TriggerOnEnteringTestSection(const FString& Section) const;
+	CORE_API bool IsAnyOnEnteringTestSectionBound() const;
+
+	/** Called leaving test section. */
+	CORE_API FOnTestSectionEvent& GetOnLeavingTestSection(const FString& Section);
+	CORE_API void TriggerOnLeavingTestSection(const FString& Section) const;
+	CORE_API bool IsAnyOnLeavingTestSectionBound() const;
+
 	/**
 	 * Return the singleton instance of the framework.
 	 *
@@ -1376,6 +1388,9 @@ private:
 	bool bForceSmokeTests;
 
 	bool bCaptureStack;
+
+	TMap<FString, FOnTestSectionEvent> OnEnteringTestSectionEvent;
+	TMap<FString, FOnTestSectionEvent> OnLeavingTestSectionEvent;
 };
 
 /** Simple abstract base class for all automation tests */
