Validate active screenshot id before accepting screenshot comparison completion

#rnx

diff --git a/Engine/Source/Developer/AutomationController/Private/AutomationControllerManager.cpp b/Engine/Source/Developer/AutomationController/Private/AutomationControllerManager.cpp
--- a/Engine/Source/Developer/AutomationController/Private/AutomationControllerManager.cpp
+++ b/Engine/Source/Developer/AutomationController/Private/AutomationControllerManager.cpp
@@ -507,7 +507,7 @@ void FAutomationControllerManager::ProcessComparisonQueue()
 			FImageComparisonResult Result = Entry->PendingComparison.Get();
 			FAutomationWorkerImageComparisonResults ResultMessage(
 				Entry->InstanceId,
-				Result.bSkipAttachingImages ? FGuid() : FGuid::NewGuid(),
+				Result.ComparisonId,
 				Result.ScreenshotPath,
 				Result.IsNew(),
 				Result.AreSimilar(),
diff --git a/Engine/Source/Developer/ScreenShotComparisonTools/Private/ScreenShotManager.cpp b/Engine/Source/Developer/ScreenShotComparisonTools/Private/ScreenShotManager.cpp
--- a/Engine/Source/Developer/ScreenShotComparisonTools/Private/ScreenShotManager.cpp
+++ b/Engine/Source/Developer/ScreenShotComparisonTools/Private/ScreenShotManager.cpp
@@ -390,6 +390,7 @@ FImageComparisonResult FScreenShotManager::CompareScreenshot(const FString& InUn
 	ComparisonResult.IdealApprovedFolderPath = IdealApprovedFolderPath;
 	// We use the subfolder path to the screenshot as name (before any environment specialization - platform, RHI - are appended).
 	ComparisonResult.ScreenshotPath = ResultsSubFolder;
+	ComparisonResult.ComparisonId = IncomingMetaData.Id;
 
 	// Do not save passing variant test screenshots
 	// Disabled for now until more variants are added, since variants now run without baseline tests in lightweight mode and we need to save at least one screenshot
diff --git a/Engine/Source/Developer/ScreenShotComparisonTools/Public/ImageComparer.h b/Engine/Source/Developer/ScreenShotComparisonTools/Public/ImageComparer.h
--- a/Engine/Source/Developer/ScreenShotComparisonTools/Public/ImageComparer.h
+++ b/Engine/Source/Developer/ScreenShotComparisonTools/Public/ImageComparer.h
@@ -308,6 +308,12 @@ public:
 	UPROPERTY()
 	FString ScreenshotPath;
 
+	/*
+		Id of the comparison request
+	*/
+	UPROPERTY()
+	FGuid ComparisonId;
+
 	/*
 		Whether to skip saving and attaching images to the report for this test
 	*/
diff --git a/Engine/Source/Programs/AutomationTool/Gauntlet/Framework/Utils/Gauntlet.HordeReport.cs b/Engine/Source/Programs/AutomationTool/Gauntlet/Framework/Utils/Gauntlet.HordeReport.cs
--- a/Engine/Source/Programs/AutomationTool/Gauntlet/Framework/Utils/Gauntlet.HordeReport.cs
+++ b/Engine/Source/Programs/AutomationTool/Gauntlet/Framework/Utils/Gauntlet.HordeReport.cs
@@ -922,7 +922,7 @@ namespace Gauntlet
 							if (!string.IsNullOrEmpty(Artifact) && Artifact.Substring(0, 4) != "0000")
 							{
 								IntArtifactEntry = InTestResult.Artifacts.Find(A => A.Id == Artifact);
-								if (IntArtifactEntry.Type == "Comparison")
+								if (IntArtifactEntry != null && IntArtifactEntry.Type == "Comparison")
 								{
 									// UE for now only produces one type of artifact that can be attached to an event and that's image comparison
 									Tag = "image comparison";
diff --git a/Engine/Source/Runtime/AutomationWorker/Private/AutomationWorkerModule.cpp b/Engine/Source/Runtime/AutomationWorker/Private/AutomationWorkerModule.cpp
--- a/Engine/Source/Runtime/AutomationWorker/Private/AutomationWorkerModule.cpp
+++ b/Engine/Source/Runtime/AutomationWorker/Private/AutomationWorkerModule.cpp
@@ -413,6 +413,12 @@ void FAutomationWorkerModule::HandleScreenShotCompared(const FAutomationWorkerIm
 {
 	UE_LOG(LogAutomationWorker, Log, TEXT("Received ScreenShotCompared from %s"), *Context->GetSender().ToString());
 
+	if (Message.UniqueId != ActiveScreenshotComparisonId)
+	{
+		UE_LOG(LogAutomationWorker, Log, TEXT("Ignoring unexpected screenshot comparison result for %s"), *Message.ScreenshotPath);
+		return;
+	}
+
 	// Image comparison finished.
 	FAutomationScreenshotCompareResults CompareResults;
 	CompareResults.UniqueId = Message.UniqueId;
@@ -489,6 +495,8 @@ void FAutomationWorkerModule::HandleScreenShotAndTraceCapturedWithName(const TAr
 
 		UE_LOG(LogAutomationWorker, Log, TEXT("Sending screenshot %s to %s"), *Message->ScreenShotName, *TestRequesterAddress.ToString());
 
+		ActiveScreenshotComparisonId = Metadata.Id;
+
 		SendMessage(Message, Message->StaticStruct(), TestRequesterAddress);
 	}
 	else
@@ -674,6 +682,7 @@ void FAutomationWorkerModule::HandleRunTestsMessage( const FAutomationWorkerRunT
 	FullTestPath = Message.FullTestPath;
 	bSendAnalytics = Message.bSendAnalytics;
 	TestRequesterAddress = Context->GetSender();
+	ActiveScreenshotComparisonId = FGuid();
 
 	// Always allow the first network command to execute
 	bExecuteNextNetworkCommand = true;
diff --git a/Engine/Source/Runtime/AutomationWorker/Private/AutomationWorkerModule.h b/Engine/Source/Runtime/AutomationWorker/Private/AutomationWorkerModule.h
--- a/Engine/Source/Runtime/AutomationWorker/Private/AutomationWorkerModule.h
+++ b/Engine/Source/Runtime/AutomationWorker/Private/AutomationWorkerModule.h
@@ -194,4 +194,7 @@ private:
 
 	/** Tracking of active section */
 	FString ActiveSection;
+
+	/** Tracking of the active screenshot comparison ID **/
+	FGuid ActiveScreenshotComparisonId;
 };
